A ClearMRT call slipped through the integration cracks. I'm not sure exactly how this no-op works, but this should at least compile.

diff --git a/Engine/Plugins/Runtime/GoogleVR/GoogleVRHMD/Source/GoogleVRHMD/Private/GoogleVRHMDRender.cpp b/Engine/Plugins/Runtime/GoogleVR/GoogleVRHMD/Source/GoogleVRHMD/Private/GoogleVRHMDRender.cpp
--- a/Engine/Plugins/Runtime/GoogleVR/GoogleVRHMD/Source/GoogleVRHMD/Private/GoogleVRHMDRender.cpp
+++ b/Engine/Plugins/Runtime/GoogleVR/GoogleVRHMD/Source/GoogleVRHMD/Private/GoogleVRHMDRender.cpp
@@ -274,7 +274,7 @@ static void ResolvePendingRenderTarget(FRHICommandListImmediate& RHICmdList, IRe
 	}
 	else
 	{
-		RHICmdList.ClearMRT(false, 0, nullptr, false, 0.0f, false, 0, FIntRect());
+		RHICmdList.ClearColorTextures(0, nullptr, nullptr, FIntRect());
 	}
 
 	RHICmdList.ImmediateFlush(EImmediateFlushType::FlushRHIThread);
