UE-42101 [daydream] Black screen with FirstPersonTemplate and 4.15P4
Duplicating Change: 3309552 in main from Release-4.15.
-Fixed creation of the rendertarget to use the msaa level on googlevr.  Also clamped the msaa level to the level that opengl reports the hardware supports.

diff --git a/Engine/Plugins/Runtime/GoogleVR/GoogleVRHMD/Source/GoogleVRHMD/Private/GoogleVRHMDRender.cpp b/Engine/Plugins/Runtime/GoogleVR/GoogleVRHMD/Source/GoogleVRHMD/Private/GoogleVRHMDRender.cpp
--- a/Engine/Plugins/Runtime/GoogleVR/GoogleVRHMD/Source/GoogleVRHMD/Private/GoogleVRHMDRender.cpp
+++ b/Engine/Plugins/Runtime/GoogleVR/GoogleVRHMD/Source/GoogleVRHMD/Private/GoogleVRHMDRender.cpp
@@ -476,6 +476,21 @@ void FGoogleVRHMDCustomPresent::Shutdown()
 	}
 }
 
+namespace {
+	int32 GetMobileMSAASampleSetting()
+	{
+		static const int32 MaxMSAASamplesSupported = FOpenGL::GetMaxMSAASamplesTileMem();
+		static const auto CVarMobileMSAA = IConsoleManager::Get().FindTConsoleVariableDataInt(TEXT("r.MobileMSAA"));
+		static const int32 CVarMobileMSAAValue = CVarMobileMSAA->GetValueOnRenderThread();
+		static const int32 MobileMSAAValue = FMath::Min(CVarMobileMSAAValue, MaxMSAASamplesSupported);
+		if (MobileMSAAValue != CVarMobileMSAAValue)
+		{
+			UE_LOG(LogHMD, Warning, TEXT("r.MobileMSAA is set to %i but we are using %i due to hardware support limitations."), CVarMobileMSAAValue, MobileMSAAValue);
+		}
+		return MobileMSAAValue;
+	}
+}
+
 bool FGoogleVRHMDCustomPresent::AllocateRenderTargetTexture(uint32 Index, uint32 SizeX, uint32 SizeY, uint8 Format, uint32 NumMips, uint32 InFlags, uint32 TargetableTextureFlags)
 {
 	FOpenGLDynamicRHI* GLRHI = static_cast<FOpenGLDynamicRHI*>(GDynamicRHI);
@@ -486,10 +501,11 @@ bool FGoogleVRHMDCustomPresent::AllocateRenderTargetTexture(uint32 Index, uint32
 		TextureSet->Resource = 0;
 	}
 
+	static int32 MobileMSAAValue = GetMobileMSAASampleSetting();
 	TextureSet = FGoogleVRHMDTexture2DSet::CreateTexture2DSet(
 		GLRHI,
 		SizeX, SizeY,
-		1, 1,
+		1, MobileMSAAValue,
 		EPixelFormat(Format),
 		TexCreate_RenderTargetable | TexCreate_ShaderResource
 		);
@@ -514,8 +530,7 @@ void FGoogleVRHMDCustomPresent::CreateGVRSwapChain()
 		return;
 	}
 
-	static const auto CVarMobileMSAA = IConsoleManager::Get().FindTConsoleVariableDataInt(TEXT("r.MobileMSAA"));
-	static int32 MobileMSAAValue = CVarMobileMSAA->GetValueOnRenderThread();
+	static int32 MobileMSAAValue = GetMobileMSAASampleSetting();
 
 	// Create resource using GVR Api
 	gvr_buffer_spec* BufferSpec = gvr_buffer_spec_create(GVRAPI);
