Disable GoogleVR plugin on Mac. Depends on OpenGLDrv module, which was disabled for Mac in CL 3337801 (copied up from Dev-Editor). Fixes UE4Game Mac compile errors.

diff --git a/Engine/Plugins/Runtime/GoogleVR/GoogleVRHMD/GoogleVRHMD.uplugin b/Engine/Plugins/Runtime/GoogleVR/GoogleVRHMD/GoogleVRHMD.uplugin
--- a/Engine/Plugins/Runtime/GoogleVR/GoogleVRHMD/GoogleVRHMD.uplugin
+++ b/Engine/Plugins/Runtime/GoogleVR/GoogleVRHMD/GoogleVRHMD.uplugin
@@ -20,7 +20,7 @@
             "Name": "GoogleVRHMD",
             "Type": "Runtime",
             "LoadingPhase": "PostConfigInit",
-            "WhitelistPlatforms": [ "Win64", "Win32", "Mac", "Android", "IOS" ]
+            "WhitelistPlatforms": [ "Win64", "Win32", "Android", "IOS" ]
         }
     ]
 }
