UE4 - Compile fixes

diff --git a/Engine/Plugins/Runtime/GoogleVR/GoogleVRHMD/Source/GoogleVRHMD/Private/GoogleVRHMDRender.cpp b/Engine/Plugins/Runtime/GoogleVR/GoogleVRHMD/Source/GoogleVRHMD/Private/GoogleVRHMDRender.cpp
--- a/Engine/Plugins/Runtime/GoogleVR/GoogleVRHMD/Source/GoogleVRHMD/Private/GoogleVRHMDRender.cpp
+++ b/Engine/Plugins/Runtime/GoogleVR/GoogleVRHMD/Source/GoogleVRHMD/Private/GoogleVRHMDRender.cpp
@@ -253,7 +253,8 @@ static void ResolvePendingRenderTarget(FRHICommandListImmediate& RHICmdList, FGr
 	}
 	else
 	{
-		RHICmdList.ClearColorTextures(0, nullptr, nullptr);
+		// What other operation can be done here?
+		//RHICmdList.ClearColorTextures(0, nullptr, nullptr);
 	}
 
 	RHICmdList.ImmediateFlush(EImmediateFlushType::FlushRHIThread);
diff --git a/Engine/Plugins/Runtime/GoogleVR/GoogleVRHMD/Source/GoogleVRHMD/Private/GoogleVRSplash.cpp b/Engine/Plugins/Runtime/GoogleVR/GoogleVRHMD/Source/GoogleVRHMD/Private/GoogleVRSplash.cpp
--- a/Engine/Plugins/Runtime/GoogleVR/GoogleVRHMD/Source/GoogleVRHMD/Private/GoogleVRSplash.cpp
+++ b/Engine/Plugins/Runtime/GoogleVR/GoogleVRHMD/Source/GoogleVRHMD/Private/GoogleVRSplash.cpp
@@ -214,6 +214,12 @@ void FGoogleVRSplash::RenderStereoSplashScreen(FRHICommandListImmediate& RHICmdL
 	{
 		RHICmdList.SetViewport(DstRect.Min.X, DstRect.Min.Y, 0, DstRect.Max.X, DstRect.Max.Y, 1.0f);
 
+		const auto FeatureLevel = GMaxRHIFeatureLevel;
+		auto ShaderMap = GetGlobalShaderMap(FeatureLevel);
+
+		TShaderMapRef<FScreenVS> VertexShader(ShaderMap);
+		TShaderMapRef<FScreenPS> PixelShader(ShaderMap);
+
 		FGraphicsPipelineStateInitializer GraphicsPSOInit;
 		RHICmdList.ApplyCachedRenderTargets(GraphicsPSOInit);
 
@@ -226,12 +232,6 @@ void FGoogleVRSplash::RenderStereoSplashScreen(FRHICommandListImmediate& RHICmdL
 		GraphicsPSOInit.BoundShaderState.PixelShaderRHI = GETSAFERHISHADER_PIXEL(*PixelShader);
 		GraphicsPSOInit.PrimitiveType = PT_TriangleList;
 
-		const auto FeatureLevel = GMaxRHIFeatureLevel;
-		auto ShaderMap = GetGlobalShaderMap(FeatureLevel);
-
-		TShaderMapRef<FScreenVS> VertexShader(ShaderMap);
-		TShaderMapRef<FScreenPS> PixelShader(ShaderMap);
-
 		SetGraphicsPipelineState(RHICmdList, GraphicsPSOInit);
 
 		PixelShader->SetParameters(RHICmdList, TStaticSamplerState<SF_Bilinear>::GetRHI(), SplashTexture->Resource->TextureRHI);
@@ -302,7 +302,7 @@ void FGoogleVRSplash::SubmitBlackFrame()
 	GVRCustomPresent->BeginRendering(GVRHMD->CachedHeadPose);
 
 	SetRenderTarget(RHICmdList, DstTexture, FTextureRHIRef());
-	RHICmdList.ClearColorTexture(DstTexture, FLinearColor(0.0f, 0.0f, 0.0f, 0.0f));
+	DrawClearQuad(RHICmdList, GMaxRHIFeatureLevel, FLinearColor(0.0f, 0.0f, 0.0f, 0.0f));
 
 	GVRCustomPresent->FinishRendering();
 }
