UE4 - Fix android compile issue

diff --git a/Engine/Plugins/Runtime/GoogleVR/GoogleVRHMD/Source/GoogleVRHMD/Private/GoogleVRSplash.cpp b/Engine/Plugins/Runtime/GoogleVR/GoogleVRHMD/Source/GoogleVRHMD/Private/GoogleVRSplash.cpp
--- a/Engine/Plugins/Runtime/GoogleVR/GoogleVRHMD/Source/GoogleVRHMD/Private/GoogleVRSplash.cpp
+++ b/Engine/Plugins/Runtime/GoogleVR/GoogleVRHMD/Source/GoogleVRHMD/Private/GoogleVRSplash.cpp
@@ -15,6 +15,7 @@
 
 #include "GoogleVRSplash.h"
 #include "PipelineStateCache.h"
+#include "ClearQuad.h"
 
 #if GOOGLEVRHMD_SUPPORTED_PLATFORMS
 #include "GoogleVRHMD.h"
@@ -206,15 +207,16 @@ void FGoogleVRSplash::RenderStereoSplashScreen(FRHICommandListImmediate& RHICmdL
 
 	FIntRect DstRect = FIntRect(0, 0, ViewportWidth, ViewportHeight);
 
+	const auto FeatureLevel = GMaxRHIFeatureLevel;
+	// Should really be SetRT and Clear
 	SetRenderTarget(RHICmdList, DstTexture, FTextureRHIRef());
-	RHICmdList.ClearColorTexture(DstTexture, FLinearColor(0.0f, 0.0f, 0.0f, 0.0f));
+	DrawClearQuad(RHICmdList, FeatureLevel, FLinearColor(0.0f, 0.0f, 0.0f, 0.0f));
 
 	// If the texture is not avaliable, we just clear the DstTexture with black.
 	if (SplashTexture && SplashTexture->IsValidLowLevel())
 	{
 		RHICmdList.SetViewport(DstRect.Min.X, DstRect.Min.Y, 0, DstRect.Max.X, DstRect.Max.Y, 1.0f);
 
-		const auto FeatureLevel = GMaxRHIFeatureLevel;
 		auto ShaderMap = GetGlobalShaderMap(FeatureLevel);
 
 		TShaderMapRef<FScreenVS> VertexShader(ShaderMap);
