Fixed compile error in GoogleVR plugin.

diff --git a/Engine/Plugins/Runtime/GoogleVR/GoogleVRHMD/Source/GoogleVRHMD/Private/GoogleVRSplash.cpp b/Engine/Plugins/Runtime/GoogleVR/GoogleVRHMD/Source/GoogleVRHMD/Private/GoogleVRSplash.cpp
--- a/Engine/Plugins/Runtime/GoogleVR/GoogleVRHMD/Source/GoogleVRHMD/Private/GoogleVRSplash.cpp
+++ b/Engine/Plugins/Runtime/GoogleVR/GoogleVRHMD/Source/GoogleVRHMD/Private/GoogleVRSplash.cpp
@@ -210,7 +210,7 @@ void FGoogleVRSplash::RenderStereoSplashScreen(FRHICommandListImmediate& RHICmdL
 	const auto FeatureLevel = GMaxRHIFeatureLevel;
 	// Should really be SetRT and Clear
 	SetRenderTarget(RHICmdList, DstTexture, FTextureRHIRef());
-	DrawClearQuad(RHICmdList, FeatureLevel, FLinearColor(0.0f, 0.0f, 0.0f, 0.0f));
+	DrawClearQuad(RHICmdList, FLinearColor(0.0f, 0.0f, 0.0f, 0.0f));
 
 	// If the texture is not avaliable, we just clear the DstTexture with black.
 	if (SplashTexture && SplashTexture->IsValidLowLevel())
@@ -304,7 +304,7 @@ void FGoogleVRSplash::SubmitBlackFrame()
 	GVRCustomPresent->BeginRendering(GVRHMD->CachedHeadPose);
 
 	SetRenderTarget(RHICmdList, DstTexture, FTextureRHIRef());
-	DrawClearQuad(RHICmdList, GMaxRHIFeatureLevel, FLinearColor(0.0f, 0.0f, 0.0f, 0.0f));
+	DrawClearQuad(RHICmdList, FLinearColor(0.0f, 0.0f, 0.0f, 0.0f));
 
 	GVRCustomPresent->FinishRendering();
 }
