Few more CIS fixes.

diff --git a/Engine/Plugins/Runtime/GoogleVR/GoogleVRHMD/Source/GoogleVRHMD/GoogleVRHMD.Build.cs b/Engine/Plugins/Runtime/GoogleVR/GoogleVRHMD/Source/GoogleVRHMD/GoogleVRHMD.Build.cs
--- a/Engine/Plugins/Runtime/GoogleVR/GoogleVRHMD/Source/GoogleVRHMD/GoogleVRHMD.Build.cs
+++ b/Engine/Plugins/Runtime/GoogleVR/GoogleVRHMD/Source/GoogleVRHMD/GoogleVRHMD.Build.cs
@@ -12,7 +12,8 @@ namespace UnrealBuildTool.Rules
 				new string[] {
 					"GoogleVRHMD/Private",
 					"../../../../../Source/Runtime/Renderer/Private",
-					"../../../../../Source/Runtime/Core/Private"
+					"../../../../../Source/Runtime/Core/Private",
+					"../../../../../Source/Runtime/ApplicationCore/Private"
 				}
 				);
 
