Properly setting up GoogleVR GL textures - filing out size info, etc. (avoiding an assert later when we free the texture).

diff --git a/Engine/Plugins/Runtime/GoogleVR/GoogleVRHMD/Source/GoogleVRHMD/Private/GoogleVRHMDRender.cpp b/Engine/Plugins/Runtime/GoogleVR/GoogleVRHMD/Source/GoogleVRHMD/Private/GoogleVRHMDRender.cpp
--- a/Engine/Plugins/Runtime/GoogleVR/GoogleVRHMD/Source/GoogleVRHMD/Private/GoogleVRHMDRender.cpp
+++ b/Engine/Plugins/Runtime/GoogleVR/GoogleVRHMD/Source/GoogleVRHMD/Private/GoogleVRHMDRender.cpp
@@ -388,6 +388,7 @@ FGoogleVRHMDTexture2DSet::FGoogleVRHMDTexture2DSet(
 	FClearValueBinding::Black
 	)
 {
+	OpenGLTextureAllocated(this, InFlags);
 }
 
 FGoogleVRHMDTexture2DSet::~FGoogleVRHMDTexture2DSet()
