Duplicating 4153647 for 4.20:
UE-60582 Mac: Daydream Instant Preview Encounters Handled Ensure On First VR Preview of Editor Run
-Some fov values from the google instant preview are sometimes reported as negative, which broke the math in our projection matrix setup resulting in inf and nan values.  Taking the absolute value of the fovs now.

diff --git a/Engine/Plugins/Runtime/GoogleVR/GoogleVRHMD/Source/GoogleVRHMD/Private/GoogleVRHMD.cpp b/Engine/Plugins/Runtime/GoogleVR/GoogleVRHMD/Source/GoogleVRHMD/Private/GoogleVRHMD.cpp
--- a/Engine/Plugins/Runtime/GoogleVR/GoogleVRHMD/Source/GoogleVRHMD/Private/GoogleVRHMD.cpp
+++ b/Engine/Plugins/Runtime/GoogleVR/GoogleVRHMD/Source/GoogleVRHMD/Private/GoogleVRHMD.cpp
@@ -1543,10 +1543,10 @@ FMatrix FGoogleVRHMD::GetStereoProjectionMatrix(const enum EStereoscopicPass Ste
 	if (bIsInstantPreviewActive) {
 		int index = (StereoPassType == eSSP_LEFT_EYE) ? 0 : 1;
 		// Have to flip left/right and top/bottom to match UE4 expectations
-		float Right = FPlatformMath::Tan(FMath::DegreesToRadians(EyeViews.eye_views[index].eye_fov.left));
-		float Left = -FPlatformMath::Tan(FMath::DegreesToRadians(EyeViews.eye_views[index].eye_fov.right));
-		float Bottom = -FPlatformMath::Tan(FMath::DegreesToRadians(EyeViews.eye_views[index].eye_fov.top));
-		float Top = FPlatformMath::Tan(FMath::DegreesToRadians(EyeViews.eye_views[index].eye_fov.bottom));
+		float Right = FPlatformMath::Tan(FMath::DegreesToRadians(FMath::Abs(EyeViews.eye_views[index].eye_fov.left)));
+		float Left = -FPlatformMath::Tan(FMath::DegreesToRadians(FMath::Abs(EyeViews.eye_views[index].eye_fov.right)));
+		float Bottom = -FPlatformMath::Tan(FMath::DegreesToRadians(FMath::Abs(EyeViews.eye_views[index].eye_fov.top)));
+		float Top = FPlatformMath::Tan(FMath::DegreesToRadians(FMath::Abs(EyeViews.eye_views[index].eye_fov.bottom)));
 		float ZNear = GNearClippingPlane;
 		float SumRL = (Right + Left);
 		float SumTB = (Top + Bottom);
