Distortion correct mesh generation needs to happen on the render thread.

diff --git a/Engine/Plugins/Runtime/GoogleVR/GoogleVRHMD/Source/GoogleVRHMD/Private/GoogleVRHMD.cpp b/Engine/Plugins/Runtime/GoogleVR/GoogleVRHMD/Source/GoogleVRHMD/Private/GoogleVRHMD.cpp
--- a/Engine/Plugins/Runtime/GoogleVR/GoogleVRHMD/Source/GoogleVRHMD/Private/GoogleVRHMD.cpp
+++ b/Engine/Plugins/Runtime/GoogleVR/GoogleVRHMD/Source/GoogleVRHMD/Private/GoogleVRHMD.cpp
@@ -928,9 +928,13 @@ void FGoogleVRHMD::SetNumOfDistortionPoints(int32 XPoints, int32 YPoints)
 	NumIndices = NumTris * 3;
 
 	// generate the distortion mesh
-	GenerateDistortionCorrectionIndexBuffer();
-	GenerateDistortionCorrectionVertexBuffer(eSSP_LEFT_EYE);
-	GenerateDistortionCorrectionVertexBuffer(eSSP_RIGHT_EYE);
+	FGoogleVRHMD * const Self = this;
+	ENQUEUE_RENDER_COMMAND(GenerateDistortionCorrectionCmd)([Self](FRHICommandListImmediate& RHICmdList)
+	{
+		Self->GenerateDistortionCorrectionIndexBuffer();
+		Self->GenerateDistortionCorrectionVertexBuffer(eSSP_LEFT_EYE);
+		Self->GenerateDistortionCorrectionVertexBuffer(eSSP_RIGHT_EYE);
+	});
 }
 
 void FGoogleVRHMD::SetDistortionMeshSize(EDistortionMeshSizeEnum MeshSize)
