Fix Proguard warning with GoogleVR plugin

#ue4
#android

diff --git a/Engine/Plugins/Runtime/GoogleVR/GoogleVRHMD/Source/GoogleVRHMD/GoogleVRHMD_APL.xml b/Engine/Plugins/Runtime/GoogleVR/GoogleVRHMD/Source/GoogleVRHMD/GoogleVRHMD_APL.xml
--- a/Engine/Plugins/Runtime/GoogleVR/GoogleVRHMD/Source/GoogleVRHMD/GoogleVRHMD_APL.xml
+++ b/Engine/Plugins/Runtime/GoogleVR/GoogleVRHMD/Source/GoogleVRHMD/GoogleVRHMD_APL.xml
@@ -125,6 +125,7 @@
 		@com.google.vr.cardboard.UsedByNative *;
 }
 
+-dontwarn com.google.vr.cardboard.R$id
 -dontwarn com.google.protobuf.nano.NanoEnumValue
 -dontwarn sun.misc.Unsafe
 		</insert>
