Merging CL 4533680 from //UE4/Dev-VR
A cleaner workaround for the Google Daydream preview crash
Don't attempt to retrieve the result of render queries early

diff --git a/Engine/Plugins/Runtime/GoogleVR/GoogleVRHMD/Source/GoogleVRHMD/Private/GoogleVRHMD.cpp b/Engine/Plugins/Runtime/GoogleVR/GoogleVRHMD/Source/GoogleVRHMD/Private/GoogleVRHMD.cpp
--- a/Engine/Plugins/Runtime/GoogleVR/GoogleVRHMD/Source/GoogleVRHMD/Private/GoogleVRHMD.cpp
+++ b/Engine/Plugins/Runtime/GoogleVR/GoogleVRHMD/Source/GoogleVRHMD/Private/GoogleVRHMD.cpp
@@ -1353,14 +1353,9 @@ void FGoogleVRHMD::PostRenderViewFamily_RenderThread(FRHICommandListImmediate& R
 	}
 
 	uint64 result = 0;
-	bool isTextureReadyForReadback = false;
-	while (SentTextureCount < ReadbackTextureCount && RHICmdList.GetRenderQueryResult(ReadbackCopyQueries[SentTextureCount % kReadbackTextureCount], result, false)) {
-		isTextureReadyForReadback = true;
-		SentTextureCount++;
-	}
-
-	if (isTextureReadyForReadback) {
-		int latestReadbackTextureIndex = (SentTextureCount - 1) % kReadbackTextureCount;
+	if (RHICmdList.GetRenderQueryResult(ReadbackCopyQueries[SentTextureCount % kReadbackTextureCount], result, false))
+	{
+		int latestReadbackTextureIndex = SentTextureCount % kReadbackTextureCount;
 		GDynamicRHI->RHIReadSurfaceData(
 			ReadbackTextures[latestReadbackTextureIndex],
 			FIntRect(FIntPoint(0, 0),
@@ -1374,6 +1369,8 @@ void FGoogleVRHMD::PostRenderViewFamily_RenderThread(FRHICommandListImmediate& R
 			ReadbackTextureSizes[latestReadbackTextureIndex].X * 4,
 			instant_preview::PIXEL_FORMAT_BGRA,
 			ReadbackReferencePoses[latestReadbackTextureIndex]);
+
+		SentTextureCount++;
 	}
 }
 #endif  // GOOGLEVRHMD_SUPPORTED_INSTANT_PREVIEW_PLATFORMS
diff --git a/Engine/Source/Runtime/Windows/D3D11RHI/Private/D3D11Query.cpp b/Engine/Source/Runtime/Windows/D3D11RHI/Private/D3D11Query.cpp
--- a/Engine/Source/Runtime/Windows/D3D11RHI/Private/D3D11Query.cpp
+++ b/Engine/Source/Runtime/Windows/D3D11RHI/Private/D3D11Query.cpp
@@ -130,10 +130,9 @@ bool FD3D11DynamicRHI::GetQueryData(ID3D11Query* Query,void* Data,SIZE_T DataSiz
 	{
 		return true;
 	}
-	else if((Result == S_FALSE || Result == DXGI_ERROR_INVALID_CALL) && !bWait)
+	else if(Result == S_FALSE && !bWait)
 	{
 		// Return failure if the query isn't complete, and waiting wasn't requested.
-		// HACK: Sometimes DXGI_ERROR_INVALID_CALL is returned instead of S_FALSE (UE-65122).
 		return false;
 	}
 	else
