HMD RenderPass updates

diff --git a/Engine/Plugins/Experimental/OSVR/Source/OSVR/Private/OSVRRender.cpp b/Engine/Plugins/Experimental/OSVR/Source/OSVR/Private/OSVRRender.cpp
--- a/Engine/Plugins/Experimental/OSVR/Source/OSVR/Private/OSVRRender.cpp
+++ b/Engine/Plugins/Experimental/OSVR/Source/OSVR/Private/OSVRRender.cpp
@@ -39,40 +39,44 @@ void FOSVRHMD::RenderTexture_RenderThread(FRHICommandListImmediate& rhiCmdList,
 	const uint32 viewportHeight = backBuffer->GetSizeY();
 
 	FGraphicsPipelineStateInitializer GraphicsPSOInit;
-	SetRenderTarget(rhiCmdList, backBuffer, FTextureRHIRef());
-	rhiCmdList.ApplyCachedRenderTargets(GraphicsPSOInit);
+	FRHIRenderPassInfo RPInfo(backBuffer, ERenderTargetActions::Load_Store);
+	rhiCmdList.BeginRenderPass(RPInfo, TEXT("OSVRHMD_RenderTexture"));
+	{
+		rhiCmdList.ApplyCachedRenderTargets(GraphicsPSOInit);
 
-	rhiCmdList.SetViewport(0, 0, 0, viewportWidth, viewportHeight, 1.0f);
+		rhiCmdList.SetViewport(0, 0, 0, viewportWidth, viewportHeight, 1.0f);
 
-	GraphicsPSOInit.BlendState = TStaticBlendState<>::GetRHI();
-	GraphicsPSOInit.RasterizerState = TStaticRasterizerState<>::GetRHI();
-	GraphicsPSOInit.DepthStencilState = TStaticDepthStencilState<false, CF_Always>::GetRHI();
+		GraphicsPSOInit.BlendState = TStaticBlendState<>::GetRHI();
+		GraphicsPSOInit.RasterizerState = TStaticRasterizerState<>::GetRHI();
+		GraphicsPSOInit.DepthStencilState = TStaticDepthStencilState<false, CF_Always>::GetRHI();
 
-	const auto featureLevel = GMaxRHIFeatureLevel;
-	auto shaderMap = GetGlobalShaderMap(featureLevel);
+		const auto featureLevel = GMaxRHIFeatureLevel;
+		auto shaderMap = GetGlobalShaderMap(featureLevel);
 
-	TShaderMapRef<FScreenVS> vertexShader(shaderMap);
-	TShaderMapRef<FScreenPS> pixelShader(shaderMap);
+		TShaderMapRef<FScreenVS> vertexShader(shaderMap);
+		TShaderMapRef<FScreenPS> pixelShader(shaderMap);
 
-	GraphicsPSOInit.BoundShaderState.VertexDeclarationRHI = RendererModule->GetFilterVertexDeclaration().VertexDeclarationRHI;
-	GraphicsPSOInit.BoundShaderState.VertexShaderRHI = GETSAFERHISHADER_VERTEX(*vertexShader);
-	GraphicsPSOInit.BoundShaderState.PixelShaderRHI = GETSAFERHISHADER_PIXEL(*pixelShader);
-	GraphicsPSOInit.PrimitiveType = PT_TriangleList;
+		GraphicsPSOInit.BoundShaderState.VertexDeclarationRHI = RendererModule->GetFilterVertexDeclaration().VertexDeclarationRHI;
+		GraphicsPSOInit.BoundShaderState.VertexShaderRHI = GETSAFERHISHADER_VERTEX(*vertexShader);
+		GraphicsPSOInit.BoundShaderState.PixelShaderRHI = GETSAFERHISHADER_PIXEL(*pixelShader);
+		GraphicsPSOInit.PrimitiveType = PT_TriangleList;
 
-	SetGraphicsPipelineState(rhiCmdList, GraphicsPSOInit);
+		SetGraphicsPipelineState(rhiCmdList, GraphicsPSOInit);
 
-	pixelShader->SetParameters(rhiCmdList, TStaticSamplerState<SF_Bilinear>::GetRHI(), srcTexture);
+		pixelShader->SetParameters(rhiCmdList, TStaticSamplerState<SF_Bilinear>::GetRHI(), srcTexture);
 
-	RendererModule->DrawRectangle(
-		rhiCmdList,
-		0, 0, // X, Y
-		viewportWidth, viewportHeight, // SizeX, SizeY
-		0.0f, 0.0f, // U, V
-		1.0f, 1.0f, // SizeU, SizeV
-		FIntPoint(viewportWidth, viewportHeight), // TargetSize
-		FIntPoint(1, 1), // TextureSize
-		*vertexShader,
-		EDRF_Default);
+		RendererModule->DrawRectangle(
+			rhiCmdList,
+			0, 0, // X, Y
+			viewportWidth, viewportHeight, // SizeX, SizeY
+			0.0f, 0.0f, // U, V
+			1.0f, 1.0f, // SizeU, SizeV
+			FIntPoint(viewportWidth, viewportHeight), // TargetSize
+			FIntPoint(1, 1), // TextureSize
+			*vertexShader,
+			EDRF_Default);
+	}
+	rhiCmdList.EndRenderPass();
 }
 
 void FOSVRHMD::GetEyeRenderParams_RenderThread(const struct FRenderingCompositePassContext& Context, FVector2D& EyeToSrcUVScaleValue, FVector2D& EyeToSrcUVOffsetValue) const
diff --git a/Engine/Plugins/Lumin/MagicLeap/Source/MagicLeap/Private/MagicLeapHMD.cpp b/Engine/Plugins/Lumin/MagicLeap/Source/MagicLeap/Private/MagicLeapHMD.cpp
--- a/Engine/Plugins/Lumin/MagicLeap/Source/MagicLeap/Private/MagicLeapHMD.cpp
+++ b/Engine/Plugins/Lumin/MagicLeap/Source/MagicLeap/Private/MagicLeapHMD.cpp
@@ -1605,56 +1605,61 @@ void FMagicLeapHMD::RenderTexture_RenderThread(class FRHICommandListImmediate& R
 		const uint32 QuadX = (ViewportWidth - BlitWidth) * 0.5f;
 		const uint32 QuadY = (ViewportHeight - BlitHeight) * 0.5f;
 
-		SetRenderTarget(RHICmdList, BackBuffer, FTextureRHIRef());
-		DrawClearQuad(RHICmdList, FLinearColor(0.0f, 0.0f, 0.0f, 1.0f));
-		RHICmdList.SetViewport(QuadX, QuadY, 0, BlitWidth + QuadX, BlitHeight + QuadY, 1.0f);
+		FRHIRenderPassInfo RPInfo(BackBuffer, ERenderTargetActions::Load_Store);
+		RHICmdList.BeginRenderPass(RPInfo, TEXT("MagicLeap_RenderTexture"));
+		{
 
-		FGraphicsPipelineStateInitializer GraphicsPSOInit;
-		RHICmdList.ApplyCachedRenderTargets(GraphicsPSOInit);
+			DrawClearQuad(RHICmdList, FLinearColor(0.0f, 0.0f, 0.0f, 1.0f));
+			RHICmdList.SetViewport(QuadX, QuadY, 0, BlitWidth + QuadX, BlitHeight + QuadY, 1.0f);
 
-		const auto FeatureLevel = GMaxRHIFeatureLevel;
-		auto ShaderMap = GetGlobalShaderMap(FeatureLevel);
-		TShaderMapRef<FScreenVS> VertexShader(ShaderMap);
-		TShaderMapRef<FScreenPS> PixelShader(ShaderMap);
+			FGraphicsPipelineStateInitializer GraphicsPSOInit;
+			RHICmdList.ApplyCachedRenderTargets(GraphicsPSOInit);
 
-		GraphicsPSOInit.BlendState = TStaticBlendState<>::GetRHI();
-		GraphicsPSOInit.RasterizerState = TStaticRasterizerState<>::GetRHI();
-		GraphicsPSOInit.DepthStencilState = TStaticDepthStencilState<false, CF_Always>::GetRHI();
-		GraphicsPSOInit.BoundShaderState.VertexDeclarationRHI = RendererModule->GetFilterVertexDeclaration().VertexDeclarationRHI;
-		GraphicsPSOInit.BoundShaderState.VertexShaderRHI = GETSAFERHISHADER_VERTEX(*VertexShader);
-		GraphicsPSOInit.BoundShaderState.PixelShaderRHI = GETSAFERHISHADER_PIXEL(*PixelShader);
-		GraphicsPSOInit.PrimitiveType = PT_TriangleList;
+			const auto FeatureLevel = GMaxRHIFeatureLevel;
+			auto ShaderMap = GetGlobalShaderMap(FeatureLevel);
+			TShaderMapRef<FScreenVS> VertexShader(ShaderMap);
+			TShaderMapRef<FScreenPS> PixelShader(ShaderMap);
 
-		SetGraphicsPipelineState(RHICmdList, GraphicsPSOInit);
+			GraphicsPSOInit.BlendState = TStaticBlendState<>::GetRHI();
+			GraphicsPSOInit.RasterizerState = TStaticRasterizerState<>::GetRHI();
+			GraphicsPSOInit.DepthStencilState = TStaticDepthStencilState<false, CF_Always>::GetRHI();
+			GraphicsPSOInit.BoundShaderState.VertexDeclarationRHI = RendererModule->GetFilterVertexDeclaration().VertexDeclarationRHI;
+			GraphicsPSOInit.BoundShaderState.VertexShaderRHI = GETSAFERHISHADER_VERTEX(*VertexShader);
+			GraphicsPSOInit.BoundShaderState.PixelShaderRHI = GETSAFERHISHADER_PIXEL(*PixelShader);
+			GraphicsPSOInit.PrimitiveType = PT_TriangleList;
 
-		PixelShader->SetParameters(RHICmdList, TStaticSamplerState<SF_Bilinear>::GetRHI(), SrcTexture);
+			SetGraphicsPipelineState(RHICmdList, GraphicsPSOInit);
 
-		if (WindowMirrorMode == 1)
-		{
-			RendererModule->DrawRectangle(
-				RHICmdList,
-				0, 0,
-				ViewportWidth, ViewportHeight,
-				0.0f, 0.0f,
-				0.5f, 1.0f,
-				FIntPoint(ViewportWidth, ViewportHeight),
-				FIntPoint(1, 1),
-				*VertexShader,
-				EDRF_Default);
-		}
-		else if (WindowMirrorMode == 2)
-		{
-			RendererModule->DrawRectangle(
-				RHICmdList,
-				0, 0,
-				ViewportWidth, ViewportHeight,
-				0.0f, 0.0f,
-				1.0f, 1.0f,
-				FIntPoint(ViewportWidth, ViewportHeight),
-				FIntPoint(1, 1),
-				*VertexShader,
-				EDRF_Default);
+			PixelShader->SetParameters(RHICmdList, TStaticSamplerState<SF_Bilinear>::GetRHI(), SrcTexture);
+
+			if (WindowMirrorMode == 1)
+			{
+				RendererModule->DrawRectangle(
+					RHICmdList,
+					0, 0,
+					ViewportWidth, ViewportHeight,
+					0.0f, 0.0f,
+					0.5f, 1.0f,
+					FIntPoint(ViewportWidth, ViewportHeight),
+					FIntPoint(1, 1),
+					*VertexShader,
+					EDRF_Default);
+			}
+			else if (WindowMirrorMode == 2)
+			{
+				RendererModule->DrawRectangle(
+					RHICmdList,
+					0, 0,
+					ViewportWidth, ViewportHeight,
+					0.0f, 0.0f,
+					1.0f, 1.0f,
+					FIntPoint(ViewportWidth, ViewportHeight),
+					FIntPoint(1, 1),
+					*VertexShader,
+					EDRF_Default);
+			}
 		}
+		RHICmdList.EndRenderPass();
 	}
 #endif //WITH_MLSDK
 }
diff --git a/Engine/Plugins/Runtime/GoogleARCore/Source/GoogleARCoreRendering/Private/GoogleARCorePassthroughCameraRenderer.cpp b/Engine/Plugins/Runtime/GoogleARCore/Source/GoogleARCoreRendering/Private/GoogleARCorePassthroughCameraRenderer.cpp
--- a/Engine/Plugins/Runtime/GoogleARCore/Source/GoogleARCoreRendering/Private/GoogleARCorePassthroughCameraRenderer.cpp
+++ b/Engine/Plugins/Runtime/GoogleARCore/Source/GoogleARCoreRendering/Private/GoogleARCorePassthroughCameraRenderer.cpp
@@ -353,36 +353,40 @@ void FGoogleARCorePassthroughCameraRenderer::CopyVideoTexture_RenderThread(FRHIC
 
 	IRendererModule& RendererModule = GetRendererModule();
 
-	SetRenderTarget(RHICmdList, DstTexture, FTextureRHIRef());
-	DrawClearQuad(RHICmdList, FLinearColor(0.0f, 0.0f, 0.0f, 0.0f));
+	FRHIRenderPassInfo RPInfo(DstTexture, ERenderTargetActions::DontLoad_Store);
+	RHICmdList.BeginRenderPass(RPInfo, TEXT("CopyVideoTexture"));
+	{
+		DrawClearQuad(RHICmdList, FLinearColor(0.0f, 0.0f, 0.0f, 0.0f));
 
-	RHICmdList.SetViewport(0, 0, 0, TargetSize.X, TargetSize.Y, 1.0f);
+		RHICmdList.SetViewport(0, 0, 0, TargetSize.X, TargetSize.Y, 1.0f);
 
-	FGraphicsPipelineStateInitializer GraphicsPSOInit;
-	RHICmdList.ApplyCachedRenderTargets(GraphicsPSOInit);
+		FGraphicsPipelineStateInitializer GraphicsPSOInit;
+		RHICmdList.ApplyCachedRenderTargets(GraphicsPSOInit);
 
-	GraphicsPSOInit.BlendState = TStaticBlendState<>::GetRHI();
-	GraphicsPSOInit.RasterizerState = TStaticRasterizerState<>::GetRHI();
-	GraphicsPSOInit.DepthStencilState = TStaticDepthStencilState<false, CF_Always>::GetRHI();
+		GraphicsPSOInit.BlendState = TStaticBlendState<>::GetRHI();
+		GraphicsPSOInit.RasterizerState = TStaticRasterizerState<>::GetRHI();
+		GraphicsPSOInit.DepthStencilState = TStaticDepthStencilState<false, CF_Always>::GetRHI();
 
-	GraphicsPSOInit.BoundShaderState.VertexDeclarationRHI = RendererModule.GetFilterVertexDeclaration().VertexDeclarationRHI;
-	GraphicsPSOInit.BoundShaderState.VertexShaderRHI = GETSAFERHISHADER_VERTEX(*VertexShader);
-	GraphicsPSOInit.BoundShaderState.PixelShaderRHI = GETSAFERHISHADER_PIXEL(*PixelShader);
-	GraphicsPSOInit.PrimitiveType = PT_TriangleList;
+		GraphicsPSOInit.BoundShaderState.VertexDeclarationRHI = RendererModule.GetFilterVertexDeclaration().VertexDeclarationRHI;
+		GraphicsPSOInit.BoundShaderState.VertexShaderRHI = GETSAFERHISHADER_VERTEX(*VertexShader);
+		GraphicsPSOInit.BoundShaderState.PixelShaderRHI = GETSAFERHISHADER_PIXEL(*PixelShader);
+		GraphicsPSOInit.PrimitiveType = PT_TriangleList;
 
-	SetGraphicsPipelineState(RHICmdList, GraphicsPSOInit);
+		SetGraphicsPipelineState(RHICmdList, GraphicsPSOInit);
 
-	PixelShader->SetParameters(RHICmdList, TStaticSamplerState<SF_Bilinear>::GetRHI(), VideoTexture.GetReference());
+		PixelShader->SetParameters(RHICmdList, TStaticSamplerState<SF_Bilinear>::GetRHI(), VideoTexture.GetReference());
 
-	RendererModule.DrawRectangle(
-		RHICmdList,
-		0, 0,
-		TargetSize.X, TargetSize.Y,
-		0, 0,
-		1, 1,
-		FIntPoint(TargetSize.X, TargetSize.Y),
-		FIntPoint(1, 1),
-		*VertexShader,
-		EDRF_Default);
+		RendererModule.DrawRectangle(
+			RHICmdList,
+			0, 0,
+			TargetSize.X, TargetSize.Y,
+			0, 0,
+			1, 1,
+			FIntPoint(TargetSize.X, TargetSize.Y),
+			FIntPoint(1, 1),
+			*VertexShader,
+			EDRF_Default);
+	}
+	RHICmdList.EndRenderPass();
 #endif
 }
diff --git a/Engine/Plugins/Runtime/GoogleVR/GoogleVRHMD/Source/GoogleVRHMD/Private/GoogleVRHMD.cpp b/Engine/Plugins/Runtime/GoogleVR/GoogleVRHMD/Source/GoogleVRHMD/Private/GoogleVRHMD.cpp
--- a/Engine/Plugins/Runtime/GoogleVR/GoogleVRHMD/Source/GoogleVRHMD/Private/GoogleVRHMD.cpp
+++ b/Engine/Plugins/Runtime/GoogleVR/GoogleVRHMD/Source/GoogleVRHMD/Private/GoogleVRHMD.cpp
@@ -1307,40 +1307,45 @@ void FGoogleVRHMD::PostRenderViewFamily_RenderThread(FRHICommandListImmediate& R
 		RendererModule->RenderTargetPoolFindFreeElement(RHICmdList, OutputDesc, ResampleTexturePooledRenderTarget, TEXT("ResampleTexture"));
 		check(ResampleTexturePooledRenderTarget);
 		const FSceneRenderTargetItem& DestRenderTarget = ResampleTexturePooledRenderTarget->GetRenderTargetItem();
-		SetRenderTarget(RHICmdList, DestRenderTarget.TargetableTexture, FTextureRHIRef());
-		RHICmdList.SetViewport(0, 0, 0.0f, ReadbackTextureSizes[textureIndex].X, ReadbackTextureSizes[textureIndex].Y, 1.0f);
-
-		FGraphicsPipelineStateInitializer GraphicsPSOInit;
-		RHICmdList.ApplyCachedRenderTargets(GraphicsPSOInit);
-
-		auto ShaderMap = GetGlobalShaderMap(FeatureLevel);
-		TShaderMapRef<FScreenVS> VertexShader(ShaderMap);
-		TShaderMapRef<FScreenPS> PixelShader(ShaderMap);
-
-		GraphicsPSOInit.BlendState = TStaticBlendState<>::GetRHI();
-		GraphicsPSOInit.RasterizerState = TStaticRasterizerState<>::GetRHI();
-		GraphicsPSOInit.DepthStencilState = TStaticDepthStencilState<false, CF_Always>::GetRHI();
-
-		GraphicsPSOInit.BoundShaderState.VertexDeclarationRHI = RendererModule->GetFilterVertexDeclaration().VertexDeclarationRHI;
-		GraphicsPSOInit.BoundShaderState.VertexShaderRHI = GETSAFERHISHADER_VERTEX(*VertexShader);
-		GraphicsPSOInit.BoundShaderState.PixelShaderRHI = GETSAFERHISHADER_PIXEL(*PixelShader);
-		GraphicsPSOInit.PrimitiveType = PT_TriangleList;
-
-		SetGraphicsPipelineState(RHICmdList, GraphicsPSOInit);
-
-
-		PixelShader->SetParameters(RHICmdList, TStaticSamplerState<SF_Point>::GetRHI(),
-			InViewFamily.RenderTarget->GetRenderTargetTexture());
-		RendererModule->DrawRectangle(
-			RHICmdList,
-			0, 0,		// Dest X, Y
-			ReadbackTextureSizes[textureIndex].X, ReadbackTextureSizes[textureIndex].Y,	// Dest Width, Height
-			0, 0,		// Source U, V
-			1, 1,		// Source USize, VSize
-			ReadbackTextureSizes[textureIndex],		// Target buffer size
-			FIntPoint(1, 1),		// Source texture size
-			*VertexShader,
-			EDRF_Default);
+
+		FRHIRenderPassInfo RPInfo(DestRenderTarget.TargetableTexture, ERenderTargetActions::Load_Store);
+		RHICmdList.BeginRenderPass(RPInfo, TEXT("GoogleVRViewFamily"));
+		{
+			RHICmdList.SetViewport(0, 0, 0.0f, ReadbackTextureSizes[textureIndex].X, ReadbackTextureSizes[textureIndex].Y, 1.0f);
+
+			FGraphicsPipelineStateInitializer GraphicsPSOInit;
+			RHICmdList.ApplyCachedRenderTargets(GraphicsPSOInit);
+
+			auto ShaderMap = GetGlobalShaderMap(FeatureLevel);
+			TShaderMapRef<FScreenVS> VertexShader(ShaderMap);
+			TShaderMapRef<FScreenPS> PixelShader(ShaderMap);
+
+			GraphicsPSOInit.BlendState = TStaticBlendState<>::GetRHI();
+			GraphicsPSOInit.RasterizerState = TStaticRasterizerState<>::GetRHI();
+			GraphicsPSOInit.DepthStencilState = TStaticDepthStencilState<false, CF_Always>::GetRHI();
+
+			GraphicsPSOInit.BoundShaderState.VertexDeclarationRHI = RendererModule->GetFilterVertexDeclaration().VertexDeclarationRHI;
+			GraphicsPSOInit.BoundShaderState.VertexShaderRHI = GETSAFERHISHADER_VERTEX(*VertexShader);
+			GraphicsPSOInit.BoundShaderState.PixelShaderRHI = GETSAFERHISHADER_PIXEL(*PixelShader);
+			GraphicsPSOInit.PrimitiveType = PT_TriangleList;
+
+			SetGraphicsPipelineState(RHICmdList, GraphicsPSOInit);
+
+
+			PixelShader->SetParameters(RHICmdList, TStaticSamplerState<SF_Point>::GetRHI(),
+				InViewFamily.RenderTarget->GetRenderTargetTexture());
+			RendererModule->DrawRectangle(
+				RHICmdList,
+				0, 0,		// Dest X, Y
+				ReadbackTextureSizes[textureIndex].X, ReadbackTextureSizes[textureIndex].Y,	// Dest Width, Height
+				0, 0,		// Source U, V
+				1, 1,		// Source USize, VSize
+				ReadbackTextureSizes[textureIndex],		// Target buffer size
+				FIntPoint(1, 1),		// Source texture size
+				*VertexShader,
+				EDRF_Default);
+		}
+		RHICmdList.EndRenderPass();
 		// Asynchronously copy delayed render target from GPU to CPU
 		RHICmdList.CopyToResolveTarget(
 			DestRenderTarget.TargetableTexture,
diff --git a/Engine/Plugins/Runtime/GoogleVR/GoogleVRHMD/Source/GoogleVRHMD/Private/GoogleVRHMDRender.cpp b/Engine/Plugins/Runtime/GoogleVR/GoogleVRHMD/Source/GoogleVRHMD/Private/GoogleVRHMDRender.cpp
--- a/Engine/Plugins/Runtime/GoogleVR/GoogleVRHMD/Source/GoogleVRHMD/Private/GoogleVRHMDRender.cpp
+++ b/Engine/Plugins/Runtime/GoogleVR/GoogleVRHMD/Source/GoogleVRHMD/Private/GoogleVRHMDRender.cpp
@@ -262,21 +262,25 @@ void FGoogleVRHMD::RenderTexture_RenderThread(FRHICommandListImmediate& RHICmdLi
 
 #if GOOGLEVRHMD_SUPPORTED_PLATFORMS
 	// When using distortion method in GVR SDK
-	if(IsUsingGVRApiDistortionCorrection() && bDistortionCorrectionEnabled)
+	if (IsUsingGVRApiDistortionCorrection() && bDistortionCorrectionEnabled)
 	{
 		// Use native gvr distortion without async reprojection
 		// Note that this method is not enabled by default.
-		if(!bUseOffscreenFramebuffers)
+		if (!bUseOffscreenFramebuffers)
 		{
 			// Set target to back buffer
-			SetRenderTarget(RHICmdList, BackBuffer, FTextureRHIRef());
-			RHICmdList.SetViewport(0, 0, 0, ViewportWidth, ViewportHeight, 1.0f);
-			ResolvePendingRenderTarget(RHICmdList, GraphicsPSOInit, RendererModule);
+			FRHIRenderPassInfo RPInfo(BackBuffer, ERenderTargetActions::Load_Store);
+			RHICmdList.BeginRenderPass(RPInfo, TEXT("GoogleVRHMD_RenderTexture"));
+			{
+				RHICmdList.SetViewport(0, 0, 0, ViewportWidth, ViewportHeight, 1.0f);
+				ResolvePendingRenderTarget(RHICmdList, GraphicsPSOInit, RendererModule);
 
-			gvr_distort_to_screen(GVRAPI, *reinterpret_cast<GLuint*>(SrcTexture->GetNativeResource()),
-								  DistortedBufferViewportList,
-								  CachedHeadPose,
-								  CachedFuturePoseTime);
+				gvr_distort_to_screen(GVRAPI, *reinterpret_cast<GLuint*>(SrcTexture->GetNativeResource()),
+					DistortedBufferViewportList,
+					CachedHeadPose,
+					CachedFuturePoseTime);
+			}
+			RHICmdList.EndRenderPass();
 		}
 		//When use aysnc reprojection, the framebuffer submit is handled in CustomPresent->FinishRendering
 	}
@@ -285,36 +289,40 @@ void FGoogleVRHMD::RenderTexture_RenderThread(FRHICommandListImmediate& RHICmdLi
 
 	// Just render directly to output
 	{
-		SetRenderTarget(RHICmdList, BackBuffer, FTextureRHIRef());
-		RHICmdList.ApplyCachedRenderTargets(GraphicsPSOInit);
+		FRHIRenderPassInfo RPInfo(BackBuffer, ERenderTargetActions::Load_Store);
+		RHICmdList.BeginRenderPass(RPInfo, TEXT("GoogleVRHMD_RenderTexture"));
+		{
+			RHICmdList.ApplyCachedRenderTargets(GraphicsPSOInit);
 
-		RHICmdList.SetViewport(0, 0, 0, ViewportWidth, ViewportHeight, 1.0f);
+			RHICmdList.SetViewport(0, 0, 0, ViewportWidth, ViewportHeight, 1.0f);
 
-		const auto FeatureLevel = GMaxRHIFeatureLevel;
-		auto ShaderMap = GetGlobalShaderMap(FeatureLevel);
+			const auto FeatureLevel = GMaxRHIFeatureLevel;
+			auto ShaderMap = GetGlobalShaderMap(FeatureLevel);
 
-		TShaderMapRef<FScreenVS> VertexShader(ShaderMap);
-		TShaderMapRef<FScreenPS> PixelShader(ShaderMap);
+			TShaderMapRef<FScreenVS> VertexShader(ShaderMap);
+			TShaderMapRef<FScreenPS> PixelShader(ShaderMap);
 
-		GraphicsPSOInit.BoundShaderState.VertexDeclarationRHI = RendererModule->GetFilterVertexDeclaration().VertexDeclarationRHI;
-		GraphicsPSOInit.BoundShaderState.VertexShaderRHI = GETSAFERHISHADER_VERTEX(*VertexShader);
-		GraphicsPSOInit.BoundShaderState.PixelShaderRHI = GETSAFERHISHADER_PIXEL(*PixelShader);
-		GraphicsPSOInit.PrimitiveType = PT_TriangleList;
+			GraphicsPSOInit.BoundShaderState.VertexDeclarationRHI = RendererModule->GetFilterVertexDeclaration().VertexDeclarationRHI;
+			GraphicsPSOInit.BoundShaderState.VertexShaderRHI = GETSAFERHISHADER_VERTEX(*VertexShader);
+			GraphicsPSOInit.BoundShaderState.PixelShaderRHI = GETSAFERHISHADER_PIXEL(*PixelShader);
+			GraphicsPSOInit.PrimitiveType = PT_TriangleList;
 
-		SetGraphicsPipelineState(RHICmdList, GraphicsPSOInit);
+			SetGraphicsPipelineState(RHICmdList, GraphicsPSOInit);
 
-		PixelShader->SetParameters(RHICmdList, TStaticSamplerState<SF_Bilinear>::GetRHI(), SrcTexture);
+			PixelShader->SetParameters(RHICmdList, TStaticSamplerState<SF_Bilinear>::GetRHI(), SrcTexture);
 
-		RendererModule->DrawRectangle(
-			RHICmdList,
-			0, 0,
-			ViewportWidth, ViewportHeight,
-			0.0f, 0.0f,
-			1.0f, 1.0f,
-			FIntPoint(ViewportWidth, ViewportHeight),
-			FIntPoint(1, 1),
-			*VertexShader,
-			EDRF_Default);
+			RendererModule->DrawRectangle(
+				RHICmdList,
+				0, 0,
+				ViewportWidth, ViewportHeight,
+				0.0f, 0.0f,
+				1.0f, 1.0f,
+				FIntPoint(ViewportWidth, ViewportHeight),
+				FIntPoint(1, 1),
+				*VertexShader,
+				EDRF_Default);
+		}
+		RHICmdList.EndRenderPass();
 	}
 }
 
diff --git a/Engine/Plugins/Runtime/GoogleVR/GoogleVRHMD/Source/GoogleVRHMD/Private/GoogleVRSplash.cpp b/Engine/Plugins/Runtime/GoogleVR/GoogleVRHMD/Source/GoogleVRHMD/Private/GoogleVRSplash.cpp
--- a/Engine/Plugins/Runtime/GoogleVR/GoogleVRHMD/Source/GoogleVRHMD/Private/GoogleVRSplash.cpp
+++ b/Engine/Plugins/Runtime/GoogleVR/GoogleVRHMD/Source/GoogleVRHMD/Private/GoogleVRSplash.cpp
@@ -205,81 +205,85 @@ void FGoogleVRSplash::RenderStereoSplashScreen(FRHICommandListImmediate& RHICmdL
 
 	const auto FeatureLevel = GMaxRHIFeatureLevel;
 	// Should really be SetRT and Clear
-	SetRenderTarget(RHICmdList, DstTexture, FTextureRHIRef());
-	DrawClearQuad(RHICmdList, FLinearColor(0.0f, 0.0f, 0.0f, 0.0f));
-
-	// If the texture is not avaliable, we just clear the DstTexture with black.
-	if (SplashTexture && SplashTexture->IsValidLowLevel())
+	FRHIRenderPassInfo RPInfo(DstTexture, ERenderTargetActions::DontLoad_Store);
+	RHICmdList.BeginRenderPass(RPInfo, TEXT("RenderStereoSplashScreen"));
 	{
-		RHICmdList.SetViewport(DstRect.Min.X, DstRect.Min.Y, 0, DstRect.Max.X, DstRect.Max.Y, 1.0f);
-
-		auto ShaderMap = GetGlobalShaderMap(FeatureLevel);
-
-		TShaderMapRef<FScreenVS> VertexShader(ShaderMap);
-		TShaderMapRef<FScreenPS> PixelShader(ShaderMap);
-
-		FGraphicsPipelineStateInitializer GraphicsPSOInit;
-		RHICmdList.ApplyCachedRenderTargets(GraphicsPSOInit);
-
-		GraphicsPSOInit.BlendState = TStaticBlendState<>::GetRHI();
-		GraphicsPSOInit.RasterizerState = TStaticRasterizerState<>::GetRHI();
-		GraphicsPSOInit.DepthStencilState = TStaticDepthStencilState<false, CF_Always>::GetRHI();
-
-		GraphicsPSOInit.BoundShaderState.VertexDeclarationRHI = RendererModule->GetFilterVertexDeclaration().VertexDeclarationRHI;
-		GraphicsPSOInit.BoundShaderState.VertexShaderRHI = GETSAFERHISHADER_VERTEX(*VertexShader);
-		GraphicsPSOInit.BoundShaderState.PixelShaderRHI = GETSAFERHISHADER_PIXEL(*PixelShader);
-		GraphicsPSOInit.PrimitiveType = PT_TriangleList;
-
-		SetGraphicsPipelineState(RHICmdList, GraphicsPSOInit);
-
-		PixelShader->SetParameters(RHICmdList, TStaticSamplerState<SF_Bilinear>::GetRHI(), SplashTexture->Resource->TextureRHI);
+		DrawClearQuad(RHICmdList, FLinearColor(0.0f, 0.0f, 0.0f, 0.0f));
 
-		// Modify the v to flip the texture. Somehow it renders flipped if not.
-		float U = SplashTextureUVOffset.X;
-		float V = SplashTextureUVOffset.Y + SplashTextureUVSize.Y;
-		float USize = SplashTextureUVSize.X;
-		float VSize = -SplashTextureUVSize.Y;
-
-		float ViewportWidthPerEye = ViewportWidth * 0.5f;
-
-		float RenderOffsetX = SplashScreenEyeOffset.X * ViewportWidthPerEye + (1.0f - RenderScale) * ViewportWidthPerEye * 0.5f;
-		float RenderOffsetY = SplashScreenEyeOffset.Y * ViewportHeight + (1.0f - RenderScale) * ViewportHeight * 0.5f;;
-		float RenderSizeX = ViewportWidthPerEye * RenderScale;
-		float RenderSizeY = ViewportHeight * RenderScale;
-
-		// Render left eye texture
-		RendererModule->DrawRectangle(
-			RHICmdList,
-			RenderOffsetX , RenderOffsetY,
-			RenderSizeX, RenderSizeY,
-			U, V,
-			USize, VSize,
-			FIntPoint(ViewportWidth, ViewportHeight),
-			FIntPoint(1, 1),
-			*VertexShader,
-			EDRF_Default);
-
-		RenderOffsetX = -SplashScreenEyeOffset.X * ViewportWidthPerEye + (1.0f - RenderScale) * ViewportWidthPerEye * 0.5f;
-
-		// In case the right eye offset is excceed the borader
-		if (RenderOffsetX < 0)
+		// If the texture is not avaliable, we just clear the DstTexture with black.
+		if (SplashTexture && SplashTexture->IsValidLowLevel())
 		{
-			U = U - RenderOffsetX / RenderSizeX;
-			RenderOffsetX = 0;
+			RHICmdList.SetViewport(DstRect.Min.X, DstRect.Min.Y, 0, DstRect.Max.X, DstRect.Max.Y, 1.0f);
+
+			auto ShaderMap = GetGlobalShaderMap(FeatureLevel);
+
+			TShaderMapRef<FScreenVS> VertexShader(ShaderMap);
+			TShaderMapRef<FScreenPS> PixelShader(ShaderMap);
+
+			FGraphicsPipelineStateInitializer GraphicsPSOInit;
+			RHICmdList.ApplyCachedRenderTargets(GraphicsPSOInit);
+
+			GraphicsPSOInit.BlendState = TStaticBlendState<>::GetRHI();
+			GraphicsPSOInit.RasterizerState = TStaticRasterizerState<>::GetRHI();
+			GraphicsPSOInit.DepthStencilState = TStaticDepthStencilState<false, CF_Always>::GetRHI();
+
+			GraphicsPSOInit.BoundShaderState.VertexDeclarationRHI = RendererModule->GetFilterVertexDeclaration().VertexDeclarationRHI;
+			GraphicsPSOInit.BoundShaderState.VertexShaderRHI = GETSAFERHISHADER_VERTEX(*VertexShader);
+			GraphicsPSOInit.BoundShaderState.PixelShaderRHI = GETSAFERHISHADER_PIXEL(*PixelShader);
+			GraphicsPSOInit.PrimitiveType = PT_TriangleList;
+
+			SetGraphicsPipelineState(RHICmdList, GraphicsPSOInit);
+
+			PixelShader->SetParameters(RHICmdList, TStaticSamplerState<SF_Bilinear>::GetRHI(), SplashTexture->Resource->TextureRHI);
+
+			// Modify the v to flip the texture. Somehow it renders flipped if not.
+			float U = SplashTextureUVOffset.X;
+			float V = SplashTextureUVOffset.Y + SplashTextureUVSize.Y;
+			float USize = SplashTextureUVSize.X;
+			float VSize = -SplashTextureUVSize.Y;
+
+			float ViewportWidthPerEye = ViewportWidth * 0.5f;
+
+			float RenderOffsetX = SplashScreenEyeOffset.X * ViewportWidthPerEye + (1.0f - RenderScale) * ViewportWidthPerEye * 0.5f;
+			float RenderOffsetY = SplashScreenEyeOffset.Y * ViewportHeight + (1.0f - RenderScale) * ViewportHeight * 0.5f;;
+			float RenderSizeX = ViewportWidthPerEye * RenderScale;
+			float RenderSizeY = ViewportHeight * RenderScale;
+
+			// Render left eye texture
+			RendererModule->DrawRectangle(
+				RHICmdList,
+				RenderOffsetX, RenderOffsetY,
+				RenderSizeX, RenderSizeY,
+				U, V,
+				USize, VSize,
+				FIntPoint(ViewportWidth, ViewportHeight),
+				FIntPoint(1, 1),
+				*VertexShader,
+				EDRF_Default);
+
+			RenderOffsetX = -SplashScreenEyeOffset.X * ViewportWidthPerEye + (1.0f - RenderScale) * ViewportWidthPerEye * 0.5f;
+
+			// In case the right eye offset is excceed the borader
+			if (RenderOffsetX < 0)
+			{
+				U = U - RenderOffsetX / RenderSizeX;
+				RenderOffsetX = 0;
+			}
+
+			// Render right eye texture
+			RendererModule->DrawRectangle(
+				RHICmdList,
+				ViewportWidthPerEye + RenderOffsetX, RenderOffsetY,
+				RenderSizeX, RenderSizeY,
+				U, V,
+				USize, VSize,
+				FIntPoint(ViewportWidth, ViewportHeight),
+				FIntPoint(1, 1),
+				*VertexShader,
+				EDRF_Default);
 		}
-
-		// Render right eye texture
-		RendererModule->DrawRectangle(
-			RHICmdList,
-			ViewportWidthPerEye + RenderOffsetX, RenderOffsetY,
-			RenderSizeX, RenderSizeY,
-			U, V,
-			USize, VSize,
-			FIntPoint(ViewportWidth, ViewportHeight),
-			FIntPoint(1, 1),
-			*VertexShader,
-			EDRF_Default);
 	}
+	RHICmdList.EndRenderPass();
 
 	// Submit frame to async reprojection
 	GVRCustomPresent->FinishRendering();
@@ -299,8 +303,12 @@ void FGoogleVRSplash::SubmitBlackFrame()
 
 	GVRCustomPresent->BeginRendering(GVRHMD->CachedHeadPose);
 
-	SetRenderTarget(RHICmdList, DstTexture, FTextureRHIRef());
-	DrawClearQuad(RHICmdList, FLinearColor(0.0f, 0.0f, 0.0f, 0.0f));
+	FRHIRenderPassInfo RPInfo(DstTexture, ERenderTargetActions::DontLoad_Store);
+	RHICmdList.BeginRenderPass(RPInfo, TEXT("SubmitBlackFrame"));
+	{
+		DrawClearQuad(RHICmdList, FLinearColor(0.0f, 0.0f, 0.0f, 0.0f));
+	}
+	RHICmdList.EndRenderPass();
 
 	GVRCustomPresent->FinishRendering();
 }
diff --git a/Engine/Plugins/Runtime/Oculus/OculusVR/Source/OculusHMD/Private/OculusHMD.cpp b/Engine/Plugins/Runtime/Oculus/OculusVR/Source/OculusHMD/Private/OculusHMD.cpp
--- a/Engine/Plugins/Runtime/Oculus/OculusVR/Source/OculusHMD/Private/OculusHMD.cpp
+++ b/Engine/Plugins/Runtime/Oculus/OculusVR/Source/OculusHMD/Private/OculusHMD.cpp
@@ -1164,10 +1164,14 @@ namespace OculusHMD
 	{
 		if (bClearBlack)
 		{
-			SetRenderTarget(RHICmdList, DstTexture, FTextureRHIRef());
-			const FIntRect ClearRect(0, 0, DstTexture->GetSizeX(), DstTexture->GetSizeY());
-			RHICmdList.SetViewport(ClearRect.Min.X, ClearRect.Min.Y, 0, ClearRect.Max.X, ClearRect.Max.Y, 1.0f);
-			DrawClearQuad(RHICmdList, FLinearColor::Black);
+			FRHIRenderPassInfo RPInfo(DstTexture, ERenderTargetActions::DontLoad_Store);
+			RHICmdList.BeginRenderPass(RPInfo, TEXT("ClearToBlack"));
+			{
+				const FIntRect ClearRect(0, 0, DstTexture->GetSizeX(), DstTexture->GetSizeY());
+				RHICmdList.SetViewport(ClearRect.Min.X, ClearRect.Min.Y, 0, ClearRect.Max.X, ClearRect.Max.Y, 1.0f);
+				DrawClearQuad(RHICmdList, FLinearColor::Black);
+			}
+			RHICmdList.EndRenderPass();
 		}
 
 		check(CustomPresent);
@@ -1660,9 +1664,13 @@ namespace OculusHMD
 			const int32 GapMinY = ViewFamily.Views[0]->UnscaledViewRect.Min.Y;
 			const int32 GapMaxY = ViewFamily.Views[1]->UnscaledViewRect.Max.Y;
 
-			SetRenderTarget(RHICmdList, ViewFamily.RenderTarget->GetRenderTargetTexture(), FTextureRHIRef());
-			RHICmdList.SetViewport(GapMinX, GapMinY, 0, GapMaxX, GapMaxY, 1.0f);
-			DrawClearQuad(RHICmdList, FLinearColor::Black);
+			FRHIRenderPassInfo RPInfo(ViewFamily.RenderTarget->GetRenderTargetTexture(), ERenderTargetActions::DontLoad_Store);
+			RHICmdList.BeginRenderPass(RPInfo, TEXT("Clear"));
+			{
+				RHICmdList.SetViewport(GapMinX, GapMinY, 0, GapMaxX, GapMaxY, 1.0f);
+				DrawClearQuad(RHICmdList, FLinearColor::Black);
+			}
+			RHICmdList.EndRenderPass();
 		}
 	#endif
 #else
diff --git a/Engine/Plugins/Runtime/Oculus/OculusVR/Source/OculusHMD/Private/OculusHMD_CustomPresent.cpp b/Engine/Plugins/Runtime/Oculus/OculusVR/Source/OculusHMD/Private/OculusHMD_CustomPresent.cpp
--- a/Engine/Plugins/Runtime/Oculus/OculusVR/Source/OculusHMD/Private/OculusHMD_CustomPresent.cpp
+++ b/Engine/Plugins/Runtime/Oculus/OculusVR/Source/OculusHMD/Private/OculusHMD_CustomPresent.cpp
@@ -397,51 +397,23 @@ void FCustomPresent::CopyTexture_RenderThread(FRHICommandListImmediate& RHICmdLi
 
 	if (DstTexture2D)
 	{
-		SetRenderTarget(RHICmdList, DstTexture, FTextureRHIRef());
-
-		if (bNoAlphaWrite)
-		{
-			RHICmdList.SetViewport(DstRect.Min.X, DstRect.Min.Y, 0.0f, DstRect.Max.X, DstRect.Max.Y, 1.0f);
-			DrawClearQuad(RHICmdList, bAlphaPremultiply ? FLinearColor::Black : FLinearColor::White);
-		}
-
-		RHICmdList.ApplyCachedRenderTargets(GraphicsPSOInit);
-
-		TShaderMapRef<FScreenPS> PixelShader(ShaderMap);
-		GraphicsPSOInit.BoundShaderState.PixelShaderRHI = GETSAFERHISHADER_PIXEL(*PixelShader);
-		SetGraphicsPipelineState(RHICmdList, GraphicsPSOInit);
-		FSamplerStateRHIParamRef SamplerState = DstRect.Size() == SrcRect.Size() ? TStaticSamplerState<SF_Point>::GetRHI() : TStaticSamplerState<SF_Bilinear>::GetRHI();
-		PixelShader->SetParameters(RHICmdList, SamplerState, SrcTextureRHI);
-
-		RHICmdList.SetViewport(DstRect.Min.X, DstRect.Min.Y, 0.0f, DstRect.Max.X, DstRect.Max.Y, 1.0f);
-
-		RendererModule->DrawRectangle(
-			RHICmdList,
-			0, 0, ViewportWidth, ViewportHeight,
-			U, V, USize, VSize,
-			TargetSize,
-			FIntPoint(1, 1),
-			*VertexShader,
-			EDRF_Default);
-	}
-	else
-	{
-		for (int FaceIndex = 0; FaceIndex < 6; FaceIndex++)
+		FRHIRenderPassInfo RPInfo(DstTexture, ERenderTargetActions::Load_Store);
+		RHICmdList.BeginRenderPass(RPInfo, TEXT("CopyTexture"));
 		{
-			SetRenderTarget(RHICmdList, DstTexture, 0, FaceIndex, FTextureRHIRef());
 
 			if (bNoAlphaWrite)
 			{
+				RHICmdList.SetViewport(DstRect.Min.X, DstRect.Min.Y, 0.0f, DstRect.Max.X, DstRect.Max.Y, 1.0f);
 				DrawClearQuad(RHICmdList, bAlphaPremultiply ? FLinearColor::Black : FLinearColor::White);
 			}
 
 			RHICmdList.ApplyCachedRenderTargets(GraphicsPSOInit);
 
-			TShaderMapRef<FOculusCubemapPS> PixelShader(ShaderMap);
+			TShaderMapRef<FScreenPS> PixelShader(ShaderMap);
 			GraphicsPSOInit.BoundShaderState.PixelShaderRHI = GETSAFERHISHADER_PIXEL(*PixelShader);
 			SetGraphicsPipelineState(RHICmdList, GraphicsPSOInit);
 			FSamplerStateRHIParamRef SamplerState = DstRect.Size() == SrcRect.Size() ? TStaticSamplerState<SF_Point>::GetRHI() : TStaticSamplerState<SF_Bilinear>::GetRHI();
-			PixelShader->SetParameters(RHICmdList, SamplerState, SrcTextureRHI, FaceIndex);
+			PixelShader->SetParameters(RHICmdList, SamplerState, SrcTextureRHI);
 
 			RHICmdList.SetViewport(DstRect.Min.X, DstRect.Min.Y, 0.0f, DstRect.Max.X, DstRect.Max.Y, 1.0f);
 
@@ -454,6 +426,43 @@ void FCustomPresent::CopyTexture_RenderThread(FRHICommandListImmediate& RHICmdLi
 				*VertexShader,
 				EDRF_Default);
 		}
+		RHICmdList.EndRenderPass();
+	}
+	else
+	{
+		for (int FaceIndex = 0; FaceIndex < 6; FaceIndex++)
+		{
+			FRHIRenderPassInfo RPInfo(DstTexture, ERenderTargetActions::Load_Store);
+			RPInfo.ColorRenderTargets[0].ArraySlice = FaceIndex;
+
+			RHICmdList.BeginRenderPass(RPInfo, TEXT("CopyTextureFace"));
+			{
+				if (bNoAlphaWrite)
+				{
+					DrawClearQuad(RHICmdList, bAlphaPremultiply ? FLinearColor::Black : FLinearColor::White);
+				}
+
+				RHICmdList.ApplyCachedRenderTargets(GraphicsPSOInit);
+
+				TShaderMapRef<FOculusCubemapPS> PixelShader(ShaderMap);
+				GraphicsPSOInit.BoundShaderState.PixelShaderRHI = GETSAFERHISHADER_PIXEL(*PixelShader);
+				SetGraphicsPipelineState(RHICmdList, GraphicsPSOInit);
+				FSamplerStateRHIParamRef SamplerState = DstRect.Size() == SrcRect.Size() ? TStaticSamplerState<SF_Point>::GetRHI() : TStaticSamplerState<SF_Bilinear>::GetRHI();
+				PixelShader->SetParameters(RHICmdList, SamplerState, SrcTextureRHI, FaceIndex);
+
+				RHICmdList.SetViewport(DstRect.Min.X, DstRect.Min.Y, 0.0f, DstRect.Max.X, DstRect.Max.Y, 1.0f);
+
+				RendererModule->DrawRectangle(
+					RHICmdList,
+					0, 0, ViewportWidth, ViewportHeight,
+					U, V, USize, VSize,
+					TargetSize,
+					FIntPoint(1, 1),
+					*VertexShader,
+					EDRF_Default);
+			}
+			RHICmdList.EndRenderPass();
+		}
 	}
 }
 
diff --git a/Engine/Plugins/Runtime/Oculus/OculusVR/Source/OculusHMD/Private/OculusHMD_StressTester.cpp b/Engine/Plugins/Runtime/Oculus/OculusVR/Source/OculusHMD/Private/OculusHMD_StressTester.cpp
--- a/Engine/Plugins/Runtime/Oculus/OculusVR/Source/OculusHMD/Private/OculusHMD_StressTester.cpp
+++ b/Engine/Plugins/Runtime/Oculus/OculusVR/Source/OculusHMD/Private/OculusHMD_StressTester.cpp
@@ -357,32 +357,35 @@ void FStressTester::DoTickGPU_RenderThread(FRHICommandListImmediate& RHICmdList,
 		}
 
 		//This is where the magic happens
-		SetRenderTarget(RHICmdList, BackBuffer, FTextureRHIRef());
-
-		FGraphicsPipelineStateInitializer GraphicsPSOInit;
-		RHICmdList.ApplyCachedRenderTargets(GraphicsPSOInit);
-		GraphicsPSOInit.BlendState = TStaticBlendState<>::GetRHI();
-		GraphicsPSOInit.RasterizerState = TStaticRasterizerState<>::GetRHI();
-		GraphicsPSOInit.DepthStencilState = TStaticDepthStencilState<false, CF_Always>::GetRHI();
-
-		const auto FeatureLevel = GMaxRHIFeatureLevel;
-		TShaderMapRef<FOculusVertexShader> VertexShader(GetGlobalShaderMap(FeatureLevel));
-		TShaderMapRef<FOculusStressShadersPS> PixelShader(GetGlobalShaderMap(FeatureLevel));
-		
-		GraphicsPSOInit.BoundShaderState.VertexDeclarationRHI = GOculusTextureVertexDeclaration.VertexDeclarationRHI;
-		GraphicsPSOInit.BoundShaderState.VertexShaderRHI = GETSAFERHISHADER_VERTEX(*VertexShader);
-		GraphicsPSOInit.BoundShaderState.PixelShaderRHI = GETSAFERHISHADER_PIXEL(*PixelShader);
-		GraphicsPSOInit.PrimitiveType = PT_TriangleStrip;
-
-		SetGraphicsPipelineState(RHICmdList, GraphicsPSOInit);
-
-		FShaderResourceViewRHIRef TextureParameterSRV = RHICreateShaderResourceView(SrcTexture, 0);
-		PixelShader->SetSurfaces(RHICmdList, TextureParameterSRV);
-		PixelShader->SetUniformBuffers(RHICmdList, ConstantParameters, VariableParameters);
-
-		// Draw a fullscreen quad that we can run our pixel shader on
-		RHICmdList.SetStreamSource(0, CreateTempOcculusVertexBuffer(), 0);
-		RHICmdList.DrawPrimitive(PT_TriangleStrip, 0, 2, 1);
+		FRHIRenderPassInfo RPInfo(BackBuffer, ERenderTargetActions::Load_Store);
+		RHICmdList.BeginRenderPass(RPInfo, TEXT("FSPass"));
+		{
+			FGraphicsPipelineStateInitializer GraphicsPSOInit;
+			RHICmdList.ApplyCachedRenderTargets(GraphicsPSOInit);
+			GraphicsPSOInit.BlendState = TStaticBlendState<>::GetRHI();
+			GraphicsPSOInit.RasterizerState = TStaticRasterizerState<>::GetRHI();
+			GraphicsPSOInit.DepthStencilState = TStaticDepthStencilState<false, CF_Always>::GetRHI();
+
+			const auto FeatureLevel = GMaxRHIFeatureLevel;
+			TShaderMapRef<FOculusVertexShader> VertexShader(GetGlobalShaderMap(FeatureLevel));
+			TShaderMapRef<FOculusStressShadersPS> PixelShader(GetGlobalShaderMap(FeatureLevel));
+
+			GraphicsPSOInit.BoundShaderState.VertexDeclarationRHI = GOculusTextureVertexDeclaration.VertexDeclarationRHI;
+			GraphicsPSOInit.BoundShaderState.VertexShaderRHI = GETSAFERHISHADER_VERTEX(*VertexShader);
+			GraphicsPSOInit.BoundShaderState.PixelShaderRHI = GETSAFERHISHADER_PIXEL(*PixelShader);
+			GraphicsPSOInit.PrimitiveType = PT_TriangleStrip;
+
+			SetGraphicsPipelineState(RHICmdList, GraphicsPSOInit);
+
+			FShaderResourceViewRHIRef TextureParameterSRV = RHICreateShaderResourceView(SrcTexture, 0);
+			PixelShader->SetSurfaces(RHICmdList, TextureParameterSRV);
+			PixelShader->SetUniformBuffers(RHICmdList, ConstantParameters, VariableParameters);
+
+			// Draw a fullscreen quad that we can run our pixel shader on
+			RHICmdList.SetStreamSource(0, CreateTempOcculusVertexBuffer(), 0);
+			RHICmdList.DrawPrimitive(PT_TriangleStrip, 0, 2, 1);
+		}
+		RHICmdList.EndRenderPass();
 
 		PixelShader->UnbindBuffers(RHICmdList);
 	}
diff --git a/Engine/Plugins/Runtime/Steam/SteamVR/Source/SteamVR/Private/SteamVRRender.cpp b/Engine/Plugins/Runtime/Steam/SteamVR/Source/SteamVR/Private/SteamVRRender.cpp
--- a/Engine/Plugins/Runtime/Steam/SteamVR/Source/SteamVR/Private/SteamVRRender.cpp
+++ b/Engine/Plugins/Runtime/Steam/SteamVR/Source/SteamVR/Private/SteamVRRender.cpp
@@ -37,8 +37,12 @@ void FSteamVRHMD::RenderTexture_RenderThread(FRHICommandListImmediate& RHICmdLis
 
 	if (bSplashIsShown)
 	{
-		SetRenderTarget(RHICmdList, SrcTexture, FTextureRHIRef());
-		DrawClearQuad(RHICmdList, FLinearColor(0, 0, 0, 0));
+		FRHIRenderPassInfo RPInfo(SrcTexture, ERenderTargetActions::DontLoad_Store);
+		RHICmdList.BeginRenderPass(RPInfo, TEXT("Clear"));
+		{
+			DrawClearQuad(RHICmdList, FLinearColor(0, 0, 0, 0));
+		}
+		RHICmdList.EndRenderPass();
 	}
 
 	check(SpectatorScreenController);
diff --git a/Engine/Plugins/Runtime/Steam/SteamVR/Source/SteamVR/Private/SteamVRSpectatorScreen.cpp b/Engine/Plugins/Runtime/Steam/SteamVR/Source/SteamVR/Private/SteamVRSpectatorScreen.cpp
--- a/Engine/Plugins/Runtime/Steam/SteamVR/Source/SteamVR/Private/SteamVRSpectatorScreen.cpp
+++ b/Engine/Plugins/Runtime/Steam/SteamVR/Source/SteamVR/Private/SteamVRSpectatorScreen.cpp
@@ -53,57 +53,60 @@ void FSteamVRHMD::CopyTexture_RenderThread(FRHICommandListImmediate& RHICmdList,
 		VSize = SrcRect.Height() / SrcTextureHeight;
 	}
 
-	SetRenderTarget(RHICmdList, DstTexture, FTextureRHIRef());
-
-	if (bClearBlack)
-	{
-		const FIntRect ClearRect(0, 0, DstTexture->GetSizeX(), DstTexture->GetSizeY());
-		RHICmdList.SetViewport(ClearRect.Min.X, ClearRect.Min.Y, 0, ClearRect.Max.X, ClearRect.Max.Y, 1.0f);
-		DrawClearQuad(RHICmdList, FLinearColor::Black);
-	}
-
-	RHICmdList.SetViewport(DstRect.Min.X, DstRect.Min.Y, 0, DstRect.Max.X, DstRect.Max.Y, 1.0f);
-
-
-	FGraphicsPipelineStateInitializer GraphicsPSOInit;
-	RHICmdList.ApplyCachedRenderTargets(GraphicsPSOInit);
-	GraphicsPSOInit.BlendState = bNoAlpha ? TStaticBlendState<>::GetRHI() : TStaticBlendState<CW_RGBA, BO_Add, BF_SourceAlpha, BF_InverseSourceAlpha, BO_Add, BF_One, BF_InverseSourceAlpha>::GetRHI();
-	GraphicsPSOInit.RasterizerState = TStaticRasterizerState<>::GetRHI();
-	GraphicsPSOInit.DepthStencilState = TStaticDepthStencilState<false, CF_Always>::GetRHI();
-	GraphicsPSOInit.PrimitiveType = PT_TriangleList;
-
-	const auto FeatureLevel = GMaxRHIFeatureLevel;
-	auto ShaderMap = GetGlobalShaderMap(FeatureLevel);
-
-	TShaderMapRef<FScreenVS> VertexShader(ShaderMap);
-	TShaderMapRef<FScreenPS> PixelShader(ShaderMap);
-
-	GraphicsPSOInit.BoundShaderState.VertexDeclarationRHI = RendererModule->GetFilterVertexDeclaration().VertexDeclarationRHI;
-	GraphicsPSOInit.BoundShaderState.VertexShaderRHI = GETSAFERHISHADER_VERTEX(*VertexShader);
-	GraphicsPSOInit.BoundShaderState.PixelShaderRHI = GETSAFERHISHADER_PIXEL(*PixelShader);
-
-	SetGraphicsPipelineState(RHICmdList, GraphicsPSOInit);
-
-	const bool bSameSize = DstRect.Size() == SrcRect.Size();
-	if (bSameSize)
+	// #todo-renderpasses Possible optimization here - use DontLoad if we will immediately clear the entire target
+	FRHIRenderPassInfo RPInfo(DstTexture, ERenderTargetActions::Load_Store);
+	RHICmdList.BeginRenderPass(RPInfo, TEXT("CopyTexture"));
 	{
-		PixelShader->SetParameters(RHICmdList, TStaticSamplerState<SF_Point>::GetRHI(), SrcTexture);
+		if (bClearBlack)
+		{
+			const FIntRect ClearRect(0, 0, DstTexture->GetSizeX(), DstTexture->GetSizeY());
+			RHICmdList.SetViewport(ClearRect.Min.X, ClearRect.Min.Y, 0, ClearRect.Max.X, ClearRect.Max.Y, 1.0f);
+			DrawClearQuad(RHICmdList, FLinearColor::Black);
+		}
+
+		RHICmdList.SetViewport(DstRect.Min.X, DstRect.Min.Y, 0, DstRect.Max.X, DstRect.Max.Y, 1.0f);
+
+		FGraphicsPipelineStateInitializer GraphicsPSOInit;
+		RHICmdList.ApplyCachedRenderTargets(GraphicsPSOInit);
+		GraphicsPSOInit.BlendState = bNoAlpha ? TStaticBlendState<>::GetRHI() : TStaticBlendState<CW_RGBA, BO_Add, BF_SourceAlpha, BF_InverseSourceAlpha, BO_Add, BF_One, BF_InverseSourceAlpha>::GetRHI();
+		GraphicsPSOInit.RasterizerState = TStaticRasterizerState<>::GetRHI();
+		GraphicsPSOInit.DepthStencilState = TStaticDepthStencilState<false, CF_Always>::GetRHI();
+		GraphicsPSOInit.PrimitiveType = PT_TriangleList;
+
+		const auto FeatureLevel = GMaxRHIFeatureLevel;
+		auto ShaderMap = GetGlobalShaderMap(FeatureLevel);
+
+		TShaderMapRef<FScreenVS> VertexShader(ShaderMap);
+		TShaderMapRef<FScreenPS> PixelShader(ShaderMap);
+
+		GraphicsPSOInit.BoundShaderState.VertexDeclarationRHI = RendererModule->GetFilterVertexDeclaration().VertexDeclarationRHI;
+		GraphicsPSOInit.BoundShaderState.VertexShaderRHI = GETSAFERHISHADER_VERTEX(*VertexShader);
+		GraphicsPSOInit.BoundShaderState.PixelShaderRHI = GETSAFERHISHADER_PIXEL(*PixelShader);
+
+		SetGraphicsPipelineState(RHICmdList, GraphicsPSOInit);
+
+		const bool bSameSize = DstRect.Size() == SrcRect.Size();
+		if (bSameSize)
+		{
+			PixelShader->SetParameters(RHICmdList, TStaticSamplerState<SF_Point>::GetRHI(), SrcTexture);
+		}
+		else
+		{
+			PixelShader->SetParameters(RHICmdList, TStaticSamplerState<SF_Bilinear>::GetRHI(), SrcTexture);
+		}
+
+		RendererModule->DrawRectangle(
+			RHICmdList,
+			0, 0,
+			ViewportWidth, ViewportHeight,
+			U, V,
+			USize, VSize,
+			TargetSize,
+			FIntPoint(1, 1),
+			*VertexShader,
+			EDRF_Default);
 	}
-	else
-	{
-		PixelShader->SetParameters(RHICmdList, TStaticSamplerState<SF_Bilinear>::GetRHI(), SrcTexture);
-	}
-
-	RendererModule->DrawRectangle(
-		RHICmdList,
-		0, 0,
-		ViewportWidth, ViewportHeight,
-		U, V,
-		USize, VSize,
-		TargetSize,
-		FIntPoint(1, 1),
-		*VertexShader,
-		EDRF_Default);
+	RHICmdList.EndRenderPass();
 }
 
 #endif // STEAMVR_SUPPORTED_PLATFORMS
diff --git a/Engine/Plugins/Runtime/WindowsMixedReality/Source/WindowsMixedRealityHMD/Private/WindowsMixedRealityHMD.cpp b/Engine/Plugins/Runtime/WindowsMixedReality/Source/WindowsMixedRealityHMD/Private/WindowsMixedRealityHMD.cpp
--- a/Engine/Plugins/Runtime/WindowsMixedReality/Source/WindowsMixedRealityHMD/Private/WindowsMixedRealityHMD.cpp
+++ b/Engine/Plugins/Runtime/WindowsMixedReality/Source/WindowsMixedRealityHMD/Private/WindowsMixedRealityHMD.cpp
@@ -783,40 +783,44 @@ namespace WindowsMixedReality
 		}
 
 		FGraphicsPipelineStateInitializer GraphicsPSOInit;
-		SetRenderTarget(RHICmdList, remappedDepthTexture, FTextureRHIRef());
-		RHICmdList.ApplyCachedRenderTargets(GraphicsPSOInit);
+		FRHIRenderPassInfo RPInfo(remappedDepthTexture, ERenderTargetActions::Load_Store);
+		RHICmdList.BeginRenderPass(RPInfo, TEXT("RemapDepth"));
+		{
+			RHICmdList.ApplyCachedRenderTargets(GraphicsPSOInit);
 
-		RHICmdList.SetViewport(0, 0, 0, viewportWidth, viewportHeight, 1.0f);
+			RHICmdList.SetViewport(0, 0, 0, viewportWidth, viewportHeight, 1.0f);
 
-		GraphicsPSOInit.BlendState = TStaticBlendState<>::GetRHI();
-		GraphicsPSOInit.RasterizerState = TStaticRasterizerState<>::GetRHI();
-		GraphicsPSOInit.DepthStencilState = TStaticDepthStencilState<false, CF_Always>::GetRHI();
+			GraphicsPSOInit.BlendState = TStaticBlendState<>::GetRHI();
+			GraphicsPSOInit.RasterizerState = TStaticRasterizerState<>::GetRHI();
+			GraphicsPSOInit.DepthStencilState = TStaticDepthStencilState<false, CF_Always>::GetRHI();
 
-		const auto featureLevel = GMaxRHIFeatureLevel;
-		auto shaderMap = GetGlobalShaderMap(featureLevel);
+			const auto featureLevel = GMaxRHIFeatureLevel;
+			auto shaderMap = GetGlobalShaderMap(featureLevel);
 
-		TShaderMapRef<FScreenVS> vertexShader(shaderMap);
-		TShaderMapRef<FDepthConversionPS> pixelShader(shaderMap);
+			TShaderMapRef<FScreenVS> vertexShader(shaderMap);
+			TShaderMapRef<FDepthConversionPS> pixelShader(shaderMap);
 
-		GraphicsPSOInit.BoundShaderState.VertexDeclarationRHI = RendererModule->GetFilterVertexDeclaration().VertexDeclarationRHI;
-		GraphicsPSOInit.BoundShaderState.VertexShaderRHI = GETSAFERHISHADER_VERTEX(*vertexShader);
-		GraphicsPSOInit.BoundShaderState.PixelShaderRHI = GETSAFERHISHADER_PIXEL(*pixelShader);
-		GraphicsPSOInit.PrimitiveType = PT_TriangleList;
+			GraphicsPSOInit.BoundShaderState.VertexDeclarationRHI = RendererModule->GetFilterVertexDeclaration().VertexDeclarationRHI;
+			GraphicsPSOInit.BoundShaderState.VertexShaderRHI = GETSAFERHISHADER_VERTEX(*vertexShader);
+			GraphicsPSOInit.BoundShaderState.PixelShaderRHI = GETSAFERHISHADER_PIXEL(*pixelShader);
+			GraphicsPSOInit.PrimitiveType = PT_TriangleList;
 
-		SetGraphicsPipelineState(RHICmdList, GraphicsPSOInit);
+			SetGraphicsPipelineState(RHICmdList, GraphicsPSOInit);
 
-		pixelShader->SetParameters(RHICmdList, farPlaneDistance / GetWorldToMetersScale(), depthFRHITexture);
+			pixelShader->SetParameters(RHICmdList, farPlaneDistance / GetWorldToMetersScale(), depthFRHITexture);
 
-		RendererModule->DrawRectangle(
-			RHICmdList,
-			0, 0, // X, Y
-			viewportWidth, viewportHeight, // SizeX, SizeY
-			0.0f, 0.0f, // U, V
-			1.0f, 1.0f, // SizeU, SizeV
-			FIntPoint(viewportWidth, viewportHeight), // TargetSize
-			FIntPoint(1, 1), // TextureSize
-			*vertexShader,
-			EDRF_Default);
+			RendererModule->DrawRectangle(
+				RHICmdList,
+				0, 0, // X, Y
+				viewportWidth, viewportHeight, // SizeX, SizeY
+				0.0f, 0.0f, // U, V
+				1.0f, 1.0f, // SizeU, SizeV
+				FIntPoint(viewportWidth, viewportHeight), // TargetSize
+				FIntPoint(1, 1), // TextureSize
+				*vertexShader,
+				EDRF_Default);
+		}
+		RHICmdList.EndRenderPass();
 
 		ID3D11Device* device = (ID3D11Device*)RHIGetNativeDevice();
 		if (device == nullptr)
diff --git a/Engine/Source/Runtime/HeadMountedDisplay/Private/DefaultStereoLayers.cpp b/Engine/Source/Runtime/HeadMountedDisplay/Private/DefaultStereoLayers.cpp
--- a/Engine/Source/Runtime/HeadMountedDisplay/Private/DefaultStereoLayers.cpp
+++ b/Engine/Source/Runtime/HeadMountedDisplay/Private/DefaultStereoLayers.cpp
@@ -241,7 +241,10 @@ void FDefaultStereoLayers::PostRenderView_RenderThread(FRHICommandListImmediate&
 	{
 		RenderTarget = InView.Family->RenderTarget->GetRenderTargetTexture();
 	}
-	SetRenderTarget(RHICmdList, RenderTarget, FTextureRHIRef());
+
+	FRHIRenderPassInfo RPInfo(RenderTarget, ERenderTargetActions::Load_Store);
+	RHICmdList.BeginRenderPass(RPInfo, TEXT("StereoLayerRender"));
+
 	RHICmdList.SetViewport(RenderParams.Viewport.Min.X, RenderParams.Viewport.Min.Y, 0, RenderParams.Viewport.Max.X, RenderParams.Viewport.Max.Y, 1.0f);
 	StereoLayerRender(RHICmdList, SortedSceneLayers, RenderParams);
 	
@@ -249,12 +252,18 @@ void FDefaultStereoLayers::PostRenderView_RenderThread(FRHICommandListImmediate&
 	FTexture2DRHIRef OverlayRenderTarget = HMDDevice->GetOverlayLayerTarget_RenderThread(InView.StereoPass, RenderParams.Viewport);
 	if (OverlayRenderTarget.IsValid())
 	{
-		SetRenderTarget(RHICmdList, OverlayRenderTarget, FTextureRHIRef());
+		RHICmdList.EndRenderPass();
+
+		FRHIRenderPassInfo RPInfoOverlayRenderTarget(OverlayRenderTarget, ERenderTargetActions::Load_Store);
+		RHICmdList.BeginRenderPass(RPInfoOverlayRenderTarget, TEXT("StereoLayerRenderIntoOverlay"));
+
 		DrawClearQuad(RHICmdList, FLinearColor(0.0f, 0.0f, 0.0f, 0.0f));
 		RHICmdList.SetViewport(RenderParams.Viewport.Min.X, RenderParams.Viewport.Min.Y, 0, RenderParams.Viewport.Max.X, RenderParams.Viewport.Max.Y, 1.0f);
 	}
 
 	StereoLayerRender(RHICmdList, SortedOverlayLayers, RenderParams);
+
+	RHICmdList.EndRenderPass();
 }
 
 bool FDefaultStereoLayers::IsActiveThisFrame(class FViewport* InViewport) const
