Fix deprectation warnings

diff --git a/Engine/Plugins/Runtime/GoogleARCore/Source/GoogleARCoreRendering/Private/GoogleARCorePassthroughCameraRenderer.cpp b/Engine/Plugins/Runtime/GoogleARCore/Source/GoogleARCoreRendering/Private/GoogleARCorePassthroughCameraRenderer.cpp
--- a/Engine/Plugins/Runtime/GoogleARCore/Source/GoogleARCoreRendering/Private/GoogleARCorePassthroughCameraRenderer.cpp
+++ b/Engine/Plugins/Runtime/GoogleARCore/Source/GoogleARCoreRendering/Private/GoogleARCorePassthroughCameraRenderer.cpp
@@ -314,7 +314,6 @@ void FGoogleARCorePassthroughCameraRenderer::RenderVideoOverlay_RenderThread(FRH
 			RHICmdList.SetStreamSource(0, OverlayVertexBufferRHI, 0);
 			RHICmdList.DrawIndexedPrimitive(
 				OverlayIndexBufferRHI,
-				PT_TriangleList,
 				/*BaseVertexIndex=*/ 0,
 				/*MinIndex=*/ 0,
 				/*NumVertices=*/ 4,
diff --git a/Engine/Plugins/Runtime/GoogleVR/GoogleVRHMD/Source/GoogleVRHMD/Private/GoogleVRHMDRender.cpp b/Engine/Plugins/Runtime/GoogleVR/GoogleVRHMD/Source/GoogleVRHMD/Private/GoogleVRHMDRender.cpp
--- a/Engine/Plugins/Runtime/GoogleVR/GoogleVRHMD/Source/GoogleVRHMD/Private/GoogleVRHMDRender.cpp
+++ b/Engine/Plugins/Runtime/GoogleVR/GoogleVRHMD/Source/GoogleVRHMD/Private/GoogleVRHMDRender.cpp
@@ -132,13 +132,13 @@ void FGoogleVRHMD::DrawDistortionMesh_RenderThread(struct FRenderingCompositePas
 	{
 		RHICmdList.SetViewport(0, 0, 0.0f, ViewportSize.X / 2, ViewportSize.Y, 1.0f);
 		RHICmdList.SetStreamSource(0, DistortionMeshVerticesLeftEye, 0);
-		RHICmdList.DrawIndexedPrimitive(DistortionMeshIndices, PT_TriangleList, 0, 0, NumVerts, 0, NumTris, 1);
+		RHICmdList.DrawIndexedPrimitive(DistortionMeshIndices, 0, 0, NumVerts, 0, NumTris, 1);
 	}
 	else
 	{
 		RHICmdList.SetViewport(ViewportSize.X / 2, 0, 0.0f, ViewportSize.X, ViewportSize.Y, 1.0f);
 		RHICmdList.SetStreamSource(0, DistortionMeshVerticesRightEye, 0);
-		RHICmdList.DrawIndexedPrimitive(DistortionMeshIndices, PT_TriangleList, 0, 0, NumVerts, 0, NumTris, 1);
+		RHICmdList.DrawIndexedPrimitive(DistortionMeshIndices, 0, 0, NumVerts, 0, NumTris, 1);
 	}
 #else
 	// Editor Preview: We are using a hardcoded quad mesh for now with no distortion applyed.
@@ -222,7 +222,6 @@ static void ResolvePendingRenderTarget(FRHICommandListImmediate& RHICmdList, FGr
 
 		RHICmdList.DrawIndexedPrimitive(
 			FakeIndexBuffer.IndexBufferRHI,
-			PT_TriangleList,
 			/*BaseVertexIndex=*/ 0,
 			/*MinIndex=*/ 0,
 			/*NumVertices=*/ 0,
