Remove iOS platform support from the GoogleVR plugin
Remove Google Cardboard as a default on all platforms

diff --git a/Engine/Plugins/Runtime/GoogleVR/GoogleVRController/GoogleVRController.uplugin b/Engine/Plugins/Runtime/GoogleVR/GoogleVRController/GoogleVRController.uplugin
--- a/Engine/Plugins/Runtime/GoogleVR/GoogleVRController/GoogleVRController.uplugin
+++ b/Engine/Plugins/Runtime/GoogleVR/GoogleVRController/GoogleVRController.uplugin
@@ -20,7 +20,7 @@
             "Name": "GoogleVRController",
             "Type": "Runtime",
             "LoadingPhase": "PostConfigInit",
-            "WhitelistPlatforms": [ "Android", "Mac", "Win32", "Win64", "Linux", "IOS" ]
+            "WhitelistPlatforms": [ "Android", "Mac", "Win32", "Win64", "Linux" ]
         }
     ],
     "Plugins": [
diff --git a/Engine/Plugins/Runtime/GoogleVR/GoogleVRHMD/GoogleVRHMD.uplugin b/Engine/Plugins/Runtime/GoogleVR/GoogleVRHMD/GoogleVRHMD.uplugin
--- a/Engine/Plugins/Runtime/GoogleVR/GoogleVRHMD/GoogleVRHMD.uplugin
+++ b/Engine/Plugins/Runtime/GoogleVR/GoogleVRHMD/GoogleVRHMD.uplugin
@@ -20,7 +20,7 @@
             "Name": "GoogleVRHMD",
             "Type": "Runtime",
             "LoadingPhase": "PostConfigInit",
-            "WhitelistPlatforms": [ "Win64", "Win32", "Mac", "Linux", "Android", "IOS" ]
+            "WhitelistPlatforms": [ "Win64", "Win32", "Mac", "Linux", "Android" ]
         }
     ]
 }
diff --git a/Engine/Plugins/Runtime/GoogleVR/GoogleVRHMD/Source/GoogleVRHMD/Private/GoogleVRHMD.cpp b/Engine/Plugins/Runtime/GoogleVR/GoogleVRHMD/Source/GoogleVRHMD/Private/GoogleVRHMD.cpp
--- a/Engine/Plugins/Runtime/GoogleVR/GoogleVRHMD/Source/GoogleVRHMD/Private/GoogleVRHMD.cpp
+++ b/Engine/Plugins/Runtime/GoogleVR/GoogleVRHMD/Source/GoogleVRHMD/Private/GoogleVRHMD.cpp
@@ -12,12 +12,6 @@
 #include "Android/AndroidApplication.h"
 #include "HAL/IConsoleManager.h"
 #endif // GOOGLEVRHMD_SUPPORTED_ANDROID_PLATFORMS
-#if GOOGLEVRHMD_SUPPORTED_IOS_PLATFORMS
-#include "IOS/IOSApplication.h"
-#include "IOS/IOSWindow.h"
-#include "IOS/IOSAppDelegate.h"
-#include "IOS/IOSView.h"
-#endif
 #if GOOGLEVRHMD_SUPPORTED_INSTANT_PREVIEW_PLATFORMS
 #include "GlobalShader.h"
 #include "ScreenRendering.h"
@@ -225,43 +219,6 @@ FString AndroidThunkCpp_GetDataString()
 
 #endif // GOOGLEVRHMD_SUPPORTED_ANDROID_PLATFORMS
 
-/////////////////////////////////////
-// Begin IOS Class Implementations //
-/////////////////////////////////////
-
-#if GOOGLEVRHMD_SUPPORTED_IOS_PLATFORMS
-
-// Helped function to get global access
-FGoogleVRHMD* GetGoogleVRHMD()
-{
-	if (GEngine->XRSystem.IsValid() && GEngine->XRSystem->GetVersionString().Contains(TEXT("GoogleVR")) )
-	{
-		return static_cast<FGoogleVRHMD*>(GEngine->XRSystem.Get());
-	}
-
-	return nullptr;
-}
-
-@implementation FOverlayViewDelegate
-
-- (void)didChangeViewerProfile
-{
-	FGoogleVRHMD* HMD = GetGoogleVRHMD();
-	if(HMD)
-	{
-		HMD->RefreshViewerProfile();
-	}
-}
-
-- (void)didTapBackButton
-{
-	bBackDetected = true;
-}
-
-@end
-
-#endif
-
 /////////////////////////////////////////////////
 // Begin FGoogleVRHMDPlugin Implementation     //
 /////////////////////////////////////////////////
@@ -323,9 +280,6 @@ FGoogleVRHMD::FGoogleVRHMD(const FAutoRegister& AutoRegister)
 	, BaseOrientation(FQuat::Identity)
 	, PixelDensity(1.0f)
 	, RendererModule(nullptr)
-#if GOOGLEVRHMD_SUPPORTED_IOS_PLATFORMS
-	, OverlayView(nil)
-#endif
 	, LastUpdatedCacheFrame(0)
 #if GOOGLEVRHMD_SUPPORTED_PLATFORMS
 	, CachedFinalHeadRotation(EForceInit::ForceInit)
@@ -482,21 +436,6 @@ FGoogleVRHMD::FGoogleVRHMD(const FAutoRegister& AutoRegister)
 
 		// Set Hardware Scaling in GvrLayout
 		AndroidThunkCpp_GvrLayout_SetFixedPresentationSurfaceSizeToCurrent();
-
-#elif GOOGLEVRHMD_SUPPORTED_IOS_PLATFORMS
-
-		// We will use Unreal's PostProcessing Distortion for iOS
-		bUseGVRApiDistortionCorrection = false;
-		bIsInDaydreamMode = false;
-
-		// Setup and show ui on iOS
-		dispatch_async(dispatch_get_main_queue(), ^
-		{
-			OverlayViewDelegate = [[FOverlayViewDelegate alloc] init];
-			OverlayView = [[GVROverlayView alloc] initWithFrame:[IOSAppDelegate GetDelegate].IOSView.bounds];
-			OverlayView.autoresizingMask = UIViewAutoresizingFlexibleWidth | UIViewAutoresizingFlexibleHeight;
-			OverlayView.delegate = OverlayViewDelegate;
-		});
 #endif // GOOGLEVRHMD_SUPPORTED_ANDROID_PLATFORMS
 
 		// By default, go ahead and disable the screen saver. The user can still change this freely
@@ -723,10 +662,6 @@ FIntPoint FGoogleVRHMD::GetUnrealMobileContentSize()
 	FPlatformRect Rect = FAndroidWindow::GetScreenRect();
 	Size.X = Rect.Right;
 	Size.Y = Rect.Bottom;
-#elif GOOGLEVRHMD_SUPPORTED_IOS_PLATFORMS
-	FPlatformRect Rect = FIOSWindow::GetScreenRect();
-	Size.X = Rect.Right;
-	Size.Y = Rect.Bottom;
 #endif
 	return Size;
 }
@@ -762,10 +697,6 @@ FIntPoint FGoogleVRHMD::SetRenderTargetSizeToDefault()
 		GVRRenderTargetSize.X = Rect.Right;
 		GVRRenderTargetSize.Y = Rect.Bottom;
 	}
-#elif GOOGLEVRHMD_SUPPORTED_IOS_PLATFORMS
-	FPlatformRect Rect = FIOSWindow::GetScreenRect();
-	GVRRenderTargetSize.X = Rect.Right;
-	GVRRenderTargetSize.Y = Rect.Bottom;
 #endif
 	return GVRRenderTargetSize;
 }
@@ -843,10 +774,6 @@ void FGoogleVRHMD::HandleGVRBackEvent()
 #if GOOGLEVRHMD_SUPPORTED_ANDROID_PLATFORMS
 		GEngine->GameViewport->Viewport->GetClient()->InputKey(GEngine->GameViewport->Viewport, 0, EKeys::Android_Back, EInputEvent::IE_Pressed);
 		GEngine->GameViewport->Viewport->GetClient()->InputKey(GEngine->GameViewport->Viewport, 0, EKeys::Android_Back, EInputEvent::IE_Released);
-#elif GOOGLEVRHMD_SUPPORTED_IOS_PLATFORMS
-		// TODO: iOS doesn't have hardware back buttons, so what should be fired?
-		GEngine->GameViewport->Viewport->GetClient()->InputKey(GEngine->GameViewport->Viewport, 0, EKeys::Global_Back, EInputEvent::IE_Pressed);
-		GEngine->GameViewport->Viewport->GetClient()->InputKey(GEngine->GameViewport->Viewport, 0, EKeys::Global_Back, EInputEvent::IE_Released);
 #endif
 	}
 }
@@ -1438,20 +1365,6 @@ bool FGoogleVRHMD::EnableStereo(bool stereo)
 	AndroidThunkCpp_UiLayer_SetEnabled(stereo);
 #endif
 
-#if GOOGLEVRHMD_SUPPORTED_IOS_PLATFORMS
-	dispatch_async(dispatch_get_main_queue(), ^
-	{
-		if (stereo)
-		{
-			[[IOSAppDelegate GetDelegate].IOSView addSubview : OverlayView];
-		}
-		else
-		{
-			[OverlayView removeFromSuperview];
-		}
-	});
-#endif
-
 	bStereoEnabled = stereo;
 	GEngine->bForceDisableFrameRateSmoothing = bStereoEnabled;
 	return bStereoEnabled;
@@ -1689,10 +1602,6 @@ bool FGoogleVRHMD::GetHMDMonitorInfo(MonitorInfo& OutMonitorInfo)
 	FPlatformRect Rect = FAndroidWindow::GetScreenRect();
 	OutMonitorInfo.ResolutionX = Rect.Right - Rect.Left;
 	OutMonitorInfo.ResolutionY = Rect.Bottom - Rect.Top;
-#elif GOOGLEVRHMD_SUPPORTED_IOS_PLATFORMS
-	FPlatformRect Rect = FIOSWindow::GetScreenRect();
-	OutMonitorInfo.ResolutionX = Rect.Right - Rect.Left;
-	OutMonitorInfo.ResolutionY = Rect.Bottom - Rect.Top;
 #else
 	gvr_recti Bounds = gvr_get_window_bounds(GVRAPI);
 	OutMonitorInfo.ResolutionX = Bounds.right - Bounds.left;
@@ -2163,22 +2072,6 @@ void FGoogleVRHMD::UpdatePoses()
 	// Convert Gvr right handed coordinate system rotation into UE left handed coordinate system.
 	CachedFinalHeadRotation = FQuat(FinalHeadPoseUnreal);
 	CachedFinalHeadRotation = FQuat(-CachedFinalHeadRotation.Z, CachedFinalHeadRotation.X, CachedFinalHeadRotation.Y, -CachedFinalHeadRotation.W);
-	
-#if GOOGLEVRHMD_SUPPORTED_IOS_PLATFORMS
-	// iOS UIInterfaceOrientationLandscapeLeft needs a 180 degree rotated transform
-	UIInterfaceOrientation Orientation = UIInterfaceOrientationPortrait;
-	Orientation = [[UIApplication sharedApplication] statusBarOrientation];
-	#if __IPHONE_OS_VERSION_MIN_REQUIRED < __IPHONE_8_0
-		Orientation = [[IOSAppDelegate GetDelegate].IOSController interfaceOrientation];
-	#endif
-	
-	if (Orientation == UIInterfaceOrientationLandscapeLeft)
-	{
-		FVector EulerRotate180(180.0f, 0.0f, 0.0f);
-		FQuat Rotate180Quat = FQuat::MakeFromEuler(EulerRotate180);
-		CachedFinalHeadRotation = Rotate180Quat * CachedFinalHeadRotation;
-	}
-#endif
 
 #elif GOOGLEVRHMD_SUPPORTED_INSTANT_PREVIEW_PLATFORMS
 	instant_preview::Session* session = ip_static_server_acquire_active_session(IpServerHandle);
diff --git a/Engine/Plugins/Runtime/GoogleVR/GoogleVRHMD/Source/GoogleVRHMD/Private/GoogleVRHMD.h b/Engine/Plugins/Runtime/GoogleVR/GoogleVRHMD/Source/GoogleVRHMD/Private/GoogleVRHMD.h
--- a/Engine/Plugins/Runtime/GoogleVR/GoogleVRHMD/Source/GoogleVRHMD/Private/GoogleVRHMD.h
+++ b/Engine/Plugins/Runtime/GoogleVR/GoogleVRHMD/Source/GoogleVRHMD/Private/GoogleVRHMD.h
@@ -46,18 +46,6 @@ extern gvr_context* GVRAPI;
 
 #endif //GOOGLEVRHMD_SUPPORTED_ANDROID_PLATFORMS
 
-#if GOOGLEVRHMD_SUPPORTED_IOS_PLATFORMS
-
-#import <GVRSDK/GVROverlayView.h>
-
-@class FOverlayViewDelegate;
-@interface FOverlayViewDelegate : UIResponder<GVROverlayViewDelegate>
-{
-}
-@end
-
-#endif //GOOGLEVRHMD_SUPPORTED_IOS_PLATFORMS
-
 // Forward decl
 class FGoogleVRHMD;
 
@@ -399,11 +387,6 @@ private:
 	FVertexBufferRHIRef DistortionMeshVerticesLeftEye;
 	FVertexBufferRHIRef DistortionMeshVerticesRightEye;
 
-#if GOOGLEVRHMD_SUPPORTED_IOS_PLATFORMS
-	GVROverlayView* OverlayView;
-	FOverlayViewDelegate* OverlayViewDelegate;
-#endif
-
 	// Cached data that should only be updated once per frame
 	mutable uint32		LastUpdatedCacheFrame;
 #if GOOGLEVRHMD_SUPPORTED_PLATFORMS
diff --git a/Engine/Plugins/Runtime/GoogleVR/GoogleVRHMD/Source/GoogleVRHMD/Private/GoogleVRHMDRender.cpp b/Engine/Plugins/Runtime/GoogleVR/GoogleVRHMD/Source/GoogleVRHMD/Private/GoogleVRHMDRender.cpp
--- a/Engine/Plugins/Runtime/GoogleVR/GoogleVRHMD/Source/GoogleVRHMD/Private/GoogleVRHMDRender.cpp
+++ b/Engine/Plugins/Runtime/GoogleVR/GoogleVRHMD/Source/GoogleVRHMD/Private/GoogleVRHMDRender.cpp
@@ -75,17 +75,7 @@ void FGoogleVRHMD::GenerateDistortionCorrectionVertexBuffer(EStereoscopicPass Ey
 				}
 			}
 
-			float ScreenYDirection = -1.0f;
-
-#if GOOGLEVRHMD_SUPPORTED_IOS_PLATFORMS && HAS_METAL
-			// Metal render the scene flipped, so we need to flip to back again.
-			bool bSupportsMetal = false;
-			GConfig->GetBool(TEXT("/Script/IOSRuntimeSettings.IOSRuntimeSettings"), TEXT("bSupportsMetal"), bSupportsMetal, GEngineIni);
-			if(bSupportsMetal)
-			{
-				ScreenYDirection = 1.0f;
-			}
-#endif
+			const float ScreenYDirection = -1.0f;
 			const FVector2D ScreenPos = FVector2D(UnDistortedCoord.X * 2.0f - 1.0f, (UnDistortedCoord.Y * 2.0f - 1.0f) * ScreenYDirection);
 
 			const FVector2D OrigRedUV = FVector2D(DistortedCoords[0].x, DistortedCoords[0].y);
@@ -97,14 +87,6 @@ void FGoogleVRHMD::GenerateDistortionCorrectionVertexBuffer(EStereoscopicPass Ey
 			FVector2D FinalGreenUV = OrigGreenUV;
 			FVector2D FinalBlueUV = OrigBlueUV;
 
-#if GOOGLEVRHMD_SUPPORTED_IOS_PLATFORMS && HAS_METAL
-			if(bSupportsMetal)
-			{
-				FinalRedUV.Y = 1.0f - FinalRedUV.Y;
-				FinalGreenUV.Y = 1.0f - FinalGreenUV.Y;
-				FinalBlueUV.Y = 1.0f - FinalBlueUV.Y;
-			}
-#endif
 			float Vignette = FMath::Clamp(XYNorm.X * kVignetteHardness, 0.0f, 1.0f)
 				* FMath::Clamp((1 - XYNorm.X)* kVignetteHardness, 0.0f, 1.0f)
 				* FMath::Clamp(XYNorm.Y * kVignetteHardness, 0.0f, 1.0f)
diff --git a/Engine/Plugins/Runtime/GoogleVR/GoogleVRHMD/Source/GoogleVRHMD/Public/IGoogleVRHMDPlugin.h b/Engine/Plugins/Runtime/GoogleVR/GoogleVRHMD/Source/GoogleVRHMD/Public/IGoogleVRHMDPlugin.h
--- a/Engine/Plugins/Runtime/GoogleVR/GoogleVRHMD/Source/GoogleVRHMD/Public/IGoogleVRHMDPlugin.h
+++ b/Engine/Plugins/Runtime/GoogleVR/GoogleVRHMD/Source/GoogleVRHMD/Public/IGoogleVRHMDPlugin.h
@@ -6,8 +6,7 @@
 #include "IHeadMountedDisplayModule.h"
 
 #define GOOGLEVRHMD_SUPPORTED_ANDROID_PLATFORMS (PLATFORM_ANDROID)
-#define GOOGLEVRHMD_SUPPORTED_IOS_PLATFORMS (PLATFORM_IOS)
-#define GOOGLEVRHMD_SUPPORTED_PLATFORMS (GOOGLEVRHMD_SUPPORTED_ANDROID_PLATFORMS || GOOGLEVRHMD_SUPPORTED_IOS_PLATFORMS)
+#define GOOGLEVRHMD_SUPPORTED_PLATFORMS (GOOGLEVRHMD_SUPPORTED_ANDROID_PLATFORMS)
 #define GOOGLEVRHMD_SUPPORTED_INSTANT_PREVIEW_PLATFORMS (!GOOGLEVRHMD_SUPPORTED_PLATFORMS && WITH_EDITOR)
 
 /**
diff --git a/Engine/Source/Runtime/Android/AndroidRuntimeSettings/Private/AndroidRuntimeSettings.cpp b/Engine/Source/Runtime/Android/AndroidRuntimeSettings/Private/AndroidRuntimeSettings.cpp
--- a/Engine/Source/Runtime/Android/AndroidRuntimeSettings/Private/AndroidRuntimeSettings.cpp
+++ b/Engine/Source/Runtime/Android/AndroidRuntimeSettings/Private/AndroidRuntimeSettings.cpp
@@ -20,7 +20,7 @@ UAndroidRuntimeSettings::UAndroidRuntimeSettings(const FObjectInitializer& Objec
 	, Orientation(EAndroidScreenOrientation::Landscape)
 	, MaxAspectRatio(2.1f)
 	, bAndroidVoiceEnabled(false)
-	, GoogleVRCaps({EGoogleVRCaps::Cardboard, EGoogleVRCaps::Daydream33})
+	, GoogleVRCaps({EGoogleVRCaps::Daydream33})
 	, bEnableGooglePlaySupport(false)
 	, bUseGetAccounts(false)
 	, bSupportAdMob(true)
