UE4.22 - Android compile fix
#rnx

diff --git a/Engine/Plugins/Runtime/GoogleVR/GoogleVRHMD/Source/GoogleVRHMD/Private/GoogleVRSplash.cpp b/Engine/Plugins/Runtime/GoogleVR/GoogleVRHMD/Source/GoogleVRHMD/Private/GoogleVRSplash.cpp
--- a/Engine/Plugins/Runtime/GoogleVR/GoogleVRHMD/Source/GoogleVRHMD/Private/GoogleVRSplash.cpp
+++ b/Engine/Plugins/Runtime/GoogleVR/GoogleVRHMD/Source/GoogleVRHMD/Private/GoogleVRSplash.cpp
@@ -113,7 +113,7 @@ void FGoogleVRSplash::Show()
 	FTickableObjectRenderThread* RenderThreadTickerLocal = RenderThreadTicker.Get();
 	FGoogleVRSplash* pGVRSplash = this;
 	ENQUEUE_RENDER_COMMAND(RegisterAsyncTick)(
-		[RenderThreadTickerLocal, pGVRSplash](FRHICommandListImmediate& RHICmdList)
+		[RenderThreadTickerLocal = RenderThreadTickerLocal, pGVRSplash](FRHICommandListImmediate& RHICmdList)
 		{
 			pGVRSplash->AllocateSplashScreenRenderTarget();
 			RenderThreadTickerLocal->Register();
