UE4.22 - Actual fix
#rnx

diff --git a/Engine/Plugins/Runtime/GoogleVR/GoogleVRHMD/Source/GoogleVRHMD/Private/GoogleVRSplash.cpp b/Engine/Plugins/Runtime/GoogleVR/GoogleVRHMD/Source/GoogleVRHMD/Private/GoogleVRSplash.cpp
--- a/Engine/Plugins/Runtime/GoogleVR/GoogleVRHMD/Source/GoogleVRHMD/Private/GoogleVRSplash.cpp
+++ b/Engine/Plugins/Runtime/GoogleVR/GoogleVRHMD/Source/GoogleVRHMD/Private/GoogleVRSplash.cpp
@@ -113,7 +113,7 @@ void FGoogleVRSplash::Show()
 	FTickableObjectRenderThread* RenderThreadTickerLocal = RenderThreadTicker.Get();
 	FGoogleVRSplash* pGVRSplash = this;
 	ENQUEUE_RENDER_COMMAND(RegisterAsyncTick)(
-		[RenderThreadTickerLocal = RenderThreadTickerLocal, pGVRSplash](FRHICommandListImmediate& RHICmdList)
+		[RenderThreadTickerLocal, pGVRSplash](FRHICommandListImmediate& RHICmdList)
 		{
 			pGVRSplash->AllocateSplashScreenRenderTarget();
 			RenderThreadTickerLocal->Register();
@@ -138,8 +138,8 @@ void FGoogleVRSplash::Hide()
 		{
 			pGVRSplash->SubmitBlackFrame();
 
-			RenderThreadTicker->Unregister();
-			RenderThreadTicker = nullptr;
+			RenderThreadTickerLocal->Unregister();
+			RenderThreadTickerLocal = nullptr;
 		});
 	FlushRenderingCommands();
 
