UE4.22 - Fix GoogleVR bug
#rnx

diff --git a/Engine/Plugins/Runtime/GoogleVR/GoogleVRHMD/Source/GoogleVRHMD/Private/GoogleVRSplash.cpp b/Engine/Plugins/Runtime/GoogleVR/GoogleVRHMD/Source/GoogleVRHMD/Private/GoogleVRSplash.cpp
--- a/Engine/Plugins/Runtime/GoogleVR/GoogleVRHMD/Source/GoogleVRHMD/Private/GoogleVRSplash.cpp
+++ b/Engine/Plugins/Runtime/GoogleVR/GoogleVRHMD/Source/GoogleVRHMD/Private/GoogleVRSplash.cpp
@@ -131,15 +131,15 @@ void FGoogleVRSplash::Hide()
 		return;
 	}
 
-	FTickableObjectRenderThread* RenderThreadTickerLocal = RenderThreadTicker.Get();
+	auto* RenderThreadTickerPtr = &RenderThreadTicker;
 	FGoogleVRSplash* pGVRSplash = this;
 	ENQUEUE_RENDER_COMMAND(UnregisterAsyncTick)(
-		[RenderThreadTickerLocal, pGVRSplash](FRHICommandListImmediate& RHICmdList)
+		[RenderThreadTickerPtr, pGVRSplash](FRHICommandListImmediate& RHICmdList)
 		{
 			pGVRSplash->SubmitBlackFrame();
 
-			RenderThreadTickerLocal->Unregister();
-			RenderThreadTickerLocal = nullptr;
+			(*RenderThreadTickerPtr)->Unregister();
+			(*RenderThreadTickerPtr) = nullptr;
 		});
 	FlushRenderingCommands();
 
