Fix Render Query leak when GoogleVR is the only VR plugin and using PIE

diff --git a/Engine/Plugins/Runtime/GoogleVR/GoogleVRHMD/Source/GoogleVRHMD/Private/GoogleVRHMD.cpp b/Engine/Plugins/Runtime/GoogleVR/GoogleVRHMD/Source/GoogleVRHMD/Private/GoogleVRHMD.cpp
--- a/Engine/Plugins/Runtime/GoogleVR/GoogleVRHMD/Source/GoogleVRHMD/Private/GoogleVRHMD.cpp
+++ b/Engine/Plugins/Runtime/GoogleVR/GoogleVRHMD/Source/GoogleVRHMD/Private/GoogleVRHMD.cpp
@@ -1299,6 +1299,8 @@ void FGoogleVRHMD::PostRenderViewFamily_RenderThread(FRHICommandListImmediate& R
 			instant_preview::PIXEL_FORMAT_BGRA,
 			ReadbackReferencePoses[latestReadbackTextureIndex]);
 
+		ReadbackCopyQueries[latestReadbackTextureIndex].ReleaseQuery();
+
 		SentTextureCount++;
 	}
 }
