Fix for editor crashing with GoogleVR plugin enabled.
- Couldn't have ReadbackCopyQueries in a structure dealloced on the main thread.  Moved dealloc to render thread.

diff --git a/Engine/Plugins/Runtime/GoogleVR/GoogleVRHMD/Source/GoogleVRHMD/Private/GoogleVRHMD.cpp b/Engine/Plugins/Runtime/GoogleVR/GoogleVRHMD/Source/GoogleVRHMD/Private/GoogleVRHMD.cpp
--- a/Engine/Plugins/Runtime/GoogleVR/GoogleVRHMD/Source/GoogleVRHMD/Private/GoogleVRHMD.cpp
+++ b/Engine/Plugins/Runtime/GoogleVR/GoogleVRHMD/Source/GoogleVRHMD/Private/GoogleVRHMD.cpp
@@ -285,6 +285,7 @@ FGoogleVRHMD::FGoogleVRHMD(const FAutoRegister& AutoRegister)
 #endif
 #if GOOGLEVRHMD_SUPPORTED_INSTANT_PREVIEW_PLATFORMS
 	, RenderQueryPool(RHICreateRenderQueryPool(RQT_AbsoluteTime, kReadbackTextureCount + 1))
+	, ReadbackCopyQueries(new FRHIPooledRenderQuery[kReadbackTextureCount])
 #endif
 	, PosePitch(0.0f)
 	, PoseYaw(0.0f)
@@ -479,6 +480,30 @@ FGoogleVRHMD::FGoogleVRHMD(const FAutoRegister& AutoRegister)
 
 FGoogleVRHMD::~FGoogleVRHMD()
 {
+#if GOOGLEVRHMD_SUPPORTED_INSTANT_PREVIEW_PLATFORMS
+	ENQUEUE_RENDER_COMMAND(ShutdownRenderThread)(
+		[this](FRHICommandListImmediate& RHICmdList)
+	{
+		if (ReadbackCopyQueries != nullptr)
+		{
+			//release any queries that are still active
+			for (int32 ClearIndex = 0; ClearIndex < kReadbackTextureCount; ClearIndex++)
+			{
+				if (ReadbackCopyQueries[ClearIndex].GetQuery() != nullptr)
+				{
+					ReadbackCopyQueries[ClearIndex].ReleaseQuery();
+				}
+			}
+
+			//make sure to deallocate on the render thread!
+			delete[] ReadbackCopyQueries;
+			ReadbackCopyQueries = nullptr;
+		}
+		RenderQueryPool.SafeRelease();
+	});
+	FlushRenderingCommands();
+#endif
+
 #if GOOGLEVRHMD_SUPPORTED_PLATFORMS
 	if (DistortedBufferViewportList)
 	{
diff --git a/Engine/Plugins/Runtime/GoogleVR/GoogleVRHMD/Source/GoogleVRHMD/Private/GoogleVRHMD.h b/Engine/Plugins/Runtime/GoogleVR/GoogleVRHMD/Source/GoogleVRHMD/Private/GoogleVRHMD.h
--- a/Engine/Plugins/Runtime/GoogleVR/GoogleVRHMD/Source/GoogleVRHMD/Private/GoogleVRHMD.h
+++ b/Engine/Plugins/Runtime/GoogleVR/GoogleVRHMD/Source/GoogleVRHMD/Private/GoogleVRHMD.h
@@ -403,7 +403,7 @@ private:
 	static const int kReadbackTextureCount = 5;
 	FTexture2DRHIRef ReadbackTextures[kReadbackTextureCount];
 	FRenderQueryPoolRHIRef RenderQueryPool;
-	FRHIPooledRenderQuery ReadbackCopyQueries[kReadbackTextureCount];
+	FRHIPooledRenderQuery* ReadbackCopyQueries = nullptr;
 	FIntPoint ReadbackTextureSizes[kReadbackTextureCount];
 	int ReadbackTextureCount;
 	instant_preview::ReferencePose ReadbackReferencePoses[kReadbackTextureCount];
