#DEV-VR - Merging CL 4275879 from Partner-Google-VR

GoogleVR 1.170 Release Update.
- Removed reference to GetHMDDeviceType.
- Updated to final version of SDK libraries.

diff --git a/Engine/Plugins/Runtime/GoogleVR/GoogleVRHMD/Source/GoogleVRHMD/Private/GoogleVRHMD.cpp b/Engine/Plugins/Runtime/GoogleVR/GoogleVRHMD/Source/GoogleVRHMD/Private/GoogleVRHMD.cpp
--- a/Engine/Plugins/Runtime/GoogleVR/GoogleVRHMD/Source/GoogleVRHMD/Private/GoogleVRHMD.cpp
+++ b/Engine/Plugins/Runtime/GoogleVR/GoogleVRHMD/Source/GoogleVRHMD/Private/GoogleVRHMD.cpp
@@ -1682,15 +1682,6 @@ void FGoogleVRHMD::EnableHMD(bool bEnable)
 	}
 }
 
-FName FGoogleVRHMD::GetHMDDeviceType() const
-{
-#if GOOGLEVRHMD_SUPPORTED_PLATFORMS
-	return FName(TEXT("ES2GenericStereoMesh"));
-#else
-	return FName(TEXT("GoogleVR"));
-#endif
-}
-
 bool FGoogleVRHMD::GetHMDMonitorInfo(MonitorInfo& OutMonitorInfo)
 {
 #if GOOGLEVRHMD_SUPPORTED_PLATFORMS
diff --git a/Engine/Plugins/Runtime/GoogleVR/GoogleVRHMD/Source/GoogleVRHMD/Private/GoogleVRHMD.h b/Engine/Plugins/Runtime/GoogleVR/GoogleVRHMD/Source/GoogleVRHMD/Private/GoogleVRHMD.h
--- a/Engine/Plugins/Runtime/GoogleVR/GoogleVRHMD/Source/GoogleVRHMD/Private/GoogleVRHMD.h
+++ b/Engine/Plugins/Runtime/GoogleVR/GoogleVRHMD/Source/GoogleVRHMD/Private/GoogleVRHMD.h
@@ -726,11 +726,6 @@ public:
 	 */
 	virtual void EnableHMD(bool bEnable = true) override;
 
-	/**
-	 * Returns the family of HMD device implemented
-	 */
-	virtual FName GetHMDDeviceType() const;
-
     /**
      * Get the name or id of the display to output for this HMD.
      */
