Fix for Google VR HMD initialization order

diff --git a/Engine/Plugins/Runtime/GoogleVR/GoogleVRHMD/Source/GoogleVRHMD/Private/GoogleVRHMD.cpp b/Engine/Plugins/Runtime/GoogleVR/GoogleVRHMD/Source/GoogleVRHMD/Private/GoogleVRHMD.cpp
--- a/Engine/Plugins/Runtime/GoogleVR/GoogleVRHMD/Source/GoogleVRHMD/Private/GoogleVRHMD.cpp
+++ b/Engine/Plugins/Runtime/GoogleVR/GoogleVRHMD/Source/GoogleVRHMD/Private/GoogleVRHMD.cpp
@@ -306,8 +306,8 @@ TSharedPtr< class IXRTrackingSystem, ESPMode::ThreadSafe > FGoogleVRHMDPlugin::C
 // Begin FGoogleVRHMD Self API     //
 /////////////////////////////////////
 FGoogleVRHMD::FGoogleVRHMD(const FAutoRegister& AutoRegister)
-	: FSceneViewExtensionBase(AutoRegister)
-	, FHeadMountedDisplayBase(nullptr)
+	: FHeadMountedDisplayBase(nullptr)
+	, FSceneViewExtensionBase(AutoRegister)
 #if GOOGLEVRHMD_SUPPORTED_PLATFORMS
 	, CustomPresent(nullptr)
 #endif
