Kills some methods of the renderer module interface that are no longer required because had moved the render target pool to RenderCore module

diff --git a/Engine/Plugins/Media/MediaIOFramework/Source/MediaIOCore/Private/MediaCapture.cpp b/Engine/Plugins/Media/MediaIOFramework/Source/MediaIOCore/Private/MediaCapture.cpp
--- a/Engine/Plugins/Media/MediaIOFramework/Source/MediaIOCore/Private/MediaCapture.cpp
+++ b/Engine/Plugins/Media/MediaIOFramework/Source/MediaIOCore/Private/MediaCapture.cpp
@@ -20,6 +20,7 @@
 #include "SceneUtils.h"
 #include "Slate/SceneViewport.h"
 #include "UObject/WeakObjectPtrTemplates.h"
+#include "RenderTargetPool.h"
 
 #if WITH_EDITOR
 #include "Editor.h"
@@ -692,7 +693,7 @@ void UMediaCapture::OnEndFrame_GameThread()
 				TexCreate_RenderTargetable,
 				false);
 			TRefCountPtr<IPooledRenderTarget> ResampleTexturePooledRenderTarget;
-			GetRendererModule().RenderTargetPoolFindFreeElement(RHICmdList, OutputDesc, ResampleTexturePooledRenderTarget, TEXT("MediaCapture"));
+			GRenderTargetPool.FindFreeElement(RHICmdList, OutputDesc, ResampleTexturePooledRenderTarget, TEXT("MediaCapture"));
 			const FSceneRenderTargetItem& DestRenderTarget = ResampleTexturePooledRenderTarget->GetRenderTargetItem();
 
 			// Do we need to crop
diff --git a/Engine/Plugins/Runtime/GoogleVR/GoogleVRHMD/Source/GoogleVRHMD/Private/GoogleVRHMD.cpp b/Engine/Plugins/Runtime/GoogleVR/GoogleVRHMD/Source/GoogleVRHMD/Private/GoogleVRHMD.cpp
--- a/Engine/Plugins/Runtime/GoogleVR/GoogleVRHMD/Source/GoogleVRHMD/Private/GoogleVRHMD.cpp
+++ b/Engine/Plugins/Runtime/GoogleVR/GoogleVRHMD/Source/GoogleVRHMD/Private/GoogleVRHMD.cpp
@@ -16,6 +16,7 @@
 #include "GlobalShader.h"
 #include "ScreenRendering.h"
 #include "PipelineStateCache.h"
+#include "RenderTargetPool.h"
 
 #include "instant_preview_server.h"
 #include "GoogleVRInstantPreviewGetServer.h"
@@ -1224,7 +1225,7 @@ void FGoogleVRHMD::PostRenderViewFamily_RenderThread(FRHICommandListImmediate& R
 		FPooledRenderTargetDesc OutputDesc(FPooledRenderTargetDesc::Create2DDesc(ReadbackTextureSizes[textureIndex], PF_B8G8R8A8, FClearValueBinding::None, TexCreate_None, TexCreate_RenderTargetable, false));
 		const auto FeatureLevel = GMaxRHIFeatureLevel;
 		TRefCountPtr<IPooledRenderTarget> ResampleTexturePooledRenderTarget;
-		RendererModule->RenderTargetPoolFindFreeElement(RHICmdList, OutputDesc, ResampleTexturePooledRenderTarget, TEXT("ResampleTexture"));
+		GRenderTargetPool.FindFreeElement(RHICmdList, OutputDesc, ResampleTexturePooledRenderTarget, TEXT("ResampleTexture"));
 		check(ResampleTexturePooledRenderTarget);
 		const FSceneRenderTargetItem& DestRenderTarget = ResampleTexturePooledRenderTarget->GetRenderTargetItem();
 
diff --git a/Engine/Source/Editor/UnrealEd/Private/EditorEngine.cpp b/Engine/Source/Editor/UnrealEd/Private/EditorEngine.cpp
--- a/Engine/Source/Editor/UnrealEd/Private/EditorEngine.cpp
+++ b/Engine/Source/Editor/UnrealEd/Private/EditorEngine.cpp
@@ -210,6 +210,7 @@
 #include "Materials/Material.h"
 #include "Materials/MaterialInstance.h"
 #include "ComponentRecreateRenderStateContext.h"
+#include "RenderTargetPool.h"
 
 DEFINE_LOG_CATEGORY_STATIC(LogEditor, Log, All);
 
@@ -1875,7 +1876,7 @@ void UEditorEngine::Tick( float DeltaSeconds, bool bIdleMode )
 					// Tick the GRenderingRealtimeClock, unless it's paused
 					GRenderingRealtimeClock.Tick(DeltaTime);
 				}
-				GetRendererModule().TickRenderTargetPool();
+				GRenderTargetPool.TickPoolElements();
 			});
 	}
 
diff --git a/Engine/Source/Runtime/Engine/Private/GameEngine.cpp b/Engine/Source/Runtime/Engine/Private/GameEngine.cpp
--- a/Engine/Source/Runtime/Engine/Private/GameEngine.cpp
+++ b/Engine/Source/Runtime/Engine/Private/GameEngine.cpp
@@ -58,6 +58,7 @@
 #include "DynamicResolutionProxy.h"
 #include "DynamicResolutionState.h"
 #include "ProfilingDebugging/CsvProfiler.h"
+#include "RenderTargetPool.h"
 
 #if WITH_EDITOR
 #include "PIEPreviewDeviceProfileSelectorModule.h"
@@ -1766,8 +1767,8 @@ void UGameEngine::Tick( float DeltaSeconds, bool bIdleMode )
 				// Tick the GRenderingRealtimeClock, unless it's paused
 				GRenderingRealtimeClock.Tick(DeltaSeconds);
 			}
-
-			GetRendererModule().TickRenderTargetPool();
+			
+			GRenderTargetPool.TickPoolElements();
 		});
 	}
 
diff --git a/Engine/Source/Runtime/Engine/Private/Rendering/SubsurfaceProfile.cpp b/Engine/Source/Runtime/Engine/Private/Rendering/SubsurfaceProfile.cpp
--- a/Engine/Source/Runtime/Engine/Private/Rendering/SubsurfaceProfile.cpp
+++ b/Engine/Source/Runtime/Engine/Private/Rendering/SubsurfaceProfile.cpp
@@ -5,6 +5,7 @@
 #include "RendererInterface.h"
 #include "Rendering/SeparableSSS.h"
 #include "EngineModule.h"
+#include "RenderTargetPool.h"
 
 DEFINE_LOG_CATEGORY_STATIC(LogSubsurfaceProfile, Log, All);
 
@@ -169,7 +170,7 @@ void FSubsurfaceProfileTexture::CreateTexture(FRHICommandListImmediate& RHICmdLi
 		Desc.Format = PF_A16B16G16R16;
 	}
 
-	GetRendererModule().RenderTargetPoolFindFreeElement(RHICmdList, Desc, GSSProfiles, TEXT("SSProfiles"));
+	GRenderTargetPool.FindFreeElement(RHICmdList, Desc, GSSProfiles, TEXT("SSProfiles"));
 
 	// Write the contents of the texture.
 	uint32 DestStride;
diff --git a/Engine/Source/Runtime/Engine/Private/Slate/DebugCanvas.cpp b/Engine/Source/Runtime/Engine/Private/Slate/DebugCanvas.cpp
--- a/Engine/Source/Runtime/Engine/Private/Slate/DebugCanvas.cpp
+++ b/Engine/Source/Runtime/Engine/Private/Slate/DebugCanvas.cpp
@@ -13,6 +13,7 @@
 #include "IXRTrackingSystem.h"
 #include "ISpectatorScreenController.h"
 #include "IHeadMountedDisplay.h"
+#include "RenderTargetPool.h"
 
 /**
  * Simple representation of the backbuffer that the debug canvas renders to
@@ -225,7 +226,7 @@ void FDebugCanvasDrawer::DrawRenderThread(FRHICommandListImmediate& RHICmdList,
 				// Set TexCreate_NoFastClear because the fast CMASK clear was not working on ps4.
 				FPooledRenderTargetDesc Desc(FPooledRenderTargetDesc::Create2DDesc(RenderThreadCanvas->GetParentCanvasSize(), PF_B8G8R8A8, FClearValueBinding(), TexCreate_SRGB, TexCreate_RenderTargetable | TexCreate_NoFastClear, false));
 				Desc.DebugName = TEXT("DebugCanvasLayerTexture");
-				GetRendererModule().RenderTargetPoolFindFreeElement(RHICmdList, Desc, LayerTexture, TEXT("DebugCanvasLayerTexture"));
+				GRenderTargetPool.FindFreeElement(RHICmdList, Desc, LayerTexture, TEXT("DebugCanvasLayerTexture"));
 				UE_LOG(LogProfilingDebugging, Log, TEXT("Allocated a %d x %d texture for HMD canvas layer"), RenderThreadCanvas->GetParentCanvasSize().X, RenderThreadCanvas->GetParentCanvasSize().Y);
 			}
 
diff --git a/Engine/Source/Runtime/MovieSceneCapture/Private/FrameGrabber.cpp b/Engine/Source/Runtime/MovieSceneCapture/Private/FrameGrabber.cpp
--- a/Engine/Source/Runtime/MovieSceneCapture/Private/FrameGrabber.cpp
+++ b/Engine/Source/Runtime/MovieSceneCapture/Private/FrameGrabber.cpp
@@ -16,6 +16,7 @@
 #include "ScreenRendering.h"
 #include "PipelineStateCache.h"
 #include "CommonRenderResources.h"
+#include "RenderTargetPool.h"
 
 int32 GFrameGrabberFrameLatency = 0;
 static FAutoConsoleVariableRef CVarFrameGrabberFrameLatency(
@@ -120,7 +121,7 @@ void FViewportSurfaceReader::ResolveRenderTarget(FViewportSurfaceReader* RenderT
 			false);
 
 		TRefCountPtr<IPooledRenderTarget> ResampleTexturePooledRenderTarget;
-		RendererModule->RenderTargetPoolFindFreeElement(RHICmdList, OutputDesc, ResampleTexturePooledRenderTarget, TEXT("ResampleTexture"));
+		GRenderTargetPool.FindFreeElement(RHICmdList, OutputDesc, ResampleTexturePooledRenderTarget, TEXT("ResampleTexture"));
 		check(ResampleTexturePooledRenderTarget);
 
 		const FSceneRenderTargetItem& DestRenderTarget = ResampleTexturePooledRenderTarget->GetRenderTargetItem();
diff --git a/Engine/Source/Runtime/RenderCore/Public/RendererInterface.h b/Engine/Source/Runtime/RenderCore/Public/RendererInterface.h
--- a/Engine/Source/Runtime/RenderCore/Public/RendererInterface.h
+++ b/Engine/Source/Runtime/RenderCore/Public/RendererInterface.h
@@ -691,14 +691,6 @@ public:
 	/** Draws a tile mesh element with the specified view. */
 	virtual void DrawTileMesh(FRHICommandListImmediate& RHICmdList, struct FMeshPassProcessorRenderState& DrawRenderState, const FSceneView& View, FMeshBatch& Mesh, bool bIsHitTesting, const class FHitProxyId& HitProxyId) = 0;
 
-	/** Render thread side, use TRefCountPtr<IPooledRenderTarget>, allows to use sharing and VisualizeTexture */
-	// TODO(RDG): Kill that guy.
-	virtual void RenderTargetPoolFindFreeElement(FRHICommandListImmediate& RHICmdList, const FPooledRenderTargetDesc& Desc, TRefCountPtr<IPooledRenderTarget> &Out, const TCHAR* InDebugName) = 0;
-	
-	/** Render thread side, to age the pool elements so they get released at some point */
-	// TODO(RDG): Kill that guy.
-	virtual void TickRenderTargetPool() = 0;
-
 	virtual const TSet<FSceneInterface*>& GetAllocatedScenes() = 0;
 
 	/** Renderer gets a chance to log some useful crash data */
diff --git a/Engine/Source/Runtime/Renderer/Private/Renderer.cpp b/Engine/Source/Runtime/Renderer/Private/Renderer.cpp
--- a/Engine/Source/Runtime/Renderer/Private/Renderer.cpp
+++ b/Engine/Source/Runtime/Renderer/Private/Renderer.cpp
@@ -328,16 +328,6 @@ void FRendererModule::DrawTileMesh(FRHICommandListImmediate& RHICmdList, FMeshPa
 	}
 }
 
-void FRendererModule::RenderTargetPoolFindFreeElement(FRHICommandListImmediate& RHICmdList, const FPooledRenderTargetDesc& Desc, TRefCountPtr<IPooledRenderTarget> &Out, const TCHAR* InDebugName)
-{
-	GRenderTargetPool.FindFreeElement(RHICmdList, Desc, Out, InDebugName);
-}
-
-void FRendererModule::TickRenderTargetPool()
-{
-	GRenderTargetPool.TickPoolElements();
-}
-
 void FRendererModule::DebugLogOnCrash()
 {
 	GVisualizeTexture.SortOrder = 1;
diff --git a/Engine/Source/Runtime/Renderer/Private/RendererModule.h b/Engine/Source/Runtime/Renderer/Private/RendererModule.h
--- a/Engine/Source/Runtime/Renderer/Private/RendererModule.h
+++ b/Engine/Source/Runtime/Renderer/Private/RendererModule.h
@@ -48,8 +48,6 @@ public:
 	virtual void SceneRenderTargetsSetBufferSize(uint32 SizeX, uint32 SizeY) override;
 	virtual void InitializeSystemTextures(FRHICommandListImmediate& RHICmdList);
 	virtual void DrawTileMesh(FRHICommandListImmediate& RHICmdList, FMeshPassProcessorRenderState& DrawRenderState, const FSceneView& View, FMeshBatch& Mesh, bool bIsHitTesting, const FHitProxyId& HitProxyId) override;
-	virtual void RenderTargetPoolFindFreeElement(FRHICommandListImmediate& RHICmdList, const FPooledRenderTargetDesc& Desc, TRefCountPtr<IPooledRenderTarget> &Out, const TCHAR* InDebugName) override;
-	virtual void TickRenderTargetPool() override;
 	virtual void DebugLogOnCrash() override;
 	virtual void GPUBenchmark(FSynthBenchmarkResults& InOut, float WorkScale) override;
 	virtual void ExecVisualizeTextureCmd(const FString& Cmd) override;
