DR - Move Google VRM use to Query pools

diff --git a/Engine/Plugins/Runtime/GoogleVR/GoogleVRHMD/Source/GoogleVRHMD/Private/GoogleVRHMD.cpp b/Engine/Plugins/Runtime/GoogleVR/GoogleVRHMD/Source/GoogleVRHMD/Private/GoogleVRHMD.cpp
--- a/Engine/Plugins/Runtime/GoogleVR/GoogleVRHMD/Source/GoogleVRHMD/Private/GoogleVRHMD.cpp
+++ b/Engine/Plugins/Runtime/GoogleVR/GoogleVRHMD/Source/GoogleVRHMD/Private/GoogleVRHMD.cpp
@@ -282,6 +282,9 @@ FGoogleVRHMD::FGoogleVRHMD(const FAutoRegister& AutoRegister)
 	, NonDistortedBufferViewportList(nullptr)
 	, ActiveViewportList(nullptr)
 	, ScratchViewport(nullptr)
+#endif
+#if GOOGLEVRHMD_SUPPORTED_INSTANT_PREVIEW_PLATFORMS
+	, RenderQueryPool(RHICreateRenderQueryPool(RQT_AbsoluteTime, kReadbackTextureCount + 1))
 #endif
 	, PosePitch(0.0f)
 	, PoseYaw(0.0f)
@@ -1213,10 +1216,10 @@ void FGoogleVRHMD::PostRenderViewFamily_RenderThread(FRHICommandListImmediate& R
 			check(ReadbackTextures[textureIndex].GetReference());
 			ReadbackTextureSizes[textureIndex] = renderSize;
 		}
-		ReadbackCopyQueries[ReadbackTextureCount % kReadbackTextureCount] = RHICmdList.CreateRenderQuery(ERenderQueryType::RQT_AbsoluteTime);
+		ReadbackCopyQueries[ReadbackTextureCount % kReadbackTextureCount] = RenderQueryPool->AllocateQuery();
 
 		// Absolute time query creation can fail on AMD hardware due to driver support
-		if (!ReadbackCopyQueries[ReadbackTextureCount % kReadbackTextureCount])
+		if (!ReadbackCopyQueries[ReadbackTextureCount % kReadbackTextureCount].GetQuery())
 		{
 			return;
 		}
@@ -1273,13 +1276,13 @@ void FGoogleVRHMD::PostRenderViewFamily_RenderThread(FRHICommandListImmediate& R
 			ReadbackTextures[ReadbackTextureCount % kReadbackTextureCount],
 			FResolveParams());
 		ReadbackReferencePoses[ReadbackTextureCount % kReadbackTextureCount] = RenderReferencePose;
-		RHICmdList.EndRenderQuery(ReadbackCopyQueries[ReadbackTextureCount % kReadbackTextureCount]);
+		RHICmdList.EndRenderQuery(ReadbackCopyQueries[ReadbackTextureCount % kReadbackTextureCount].GetQuery());
 
 		ReadbackTextureCount++;
 	}
 
 	uint64 result = 0;
-	if (RHICmdList.GetRenderQueryResult(ReadbackCopyQueries[SentTextureCount % kReadbackTextureCount], result, false))
+	if (RHICmdList.GetRenderQueryResult(ReadbackCopyQueries[SentTextureCount % kReadbackTextureCount].GetQuery(), result, false))
 	{
 		int latestReadbackTextureIndex = SentTextureCount % kReadbackTextureCount;
 		GDynamicRHI->RHIReadSurfaceData(
diff --git a/Engine/Plugins/Runtime/GoogleVR/GoogleVRHMD/Source/GoogleVRHMD/Private/GoogleVRHMD.h b/Engine/Plugins/Runtime/GoogleVR/GoogleVRHMD/Source/GoogleVRHMD/Private/GoogleVRHMD.h
--- a/Engine/Plugins/Runtime/GoogleVR/GoogleVRHMD/Source/GoogleVRHMD/Private/GoogleVRHMD.h
+++ b/Engine/Plugins/Runtime/GoogleVR/GoogleVRHMD/Source/GoogleVRHMD/Private/GoogleVRHMD.h
@@ -402,7 +402,8 @@ private:
 #if GOOGLEVRHMD_SUPPORTED_INSTANT_PREVIEW_PLATFORMS
 	static const int kReadbackTextureCount = 5;
 	FTexture2DRHIRef ReadbackTextures[kReadbackTextureCount];
-	FRenderQueryRHIRef ReadbackCopyQueries[kReadbackTextureCount];
+	FRenderQueryPoolRHIRef RenderQueryPool;
+	FRHIPooledRenderQuery ReadbackCopyQueries[kReadbackTextureCount];
 	FIntPoint ReadbackTextureSizes[kReadbackTextureCount];
 	int ReadbackTextureCount;
 	instant_preview::ReferencePose ReadbackReferencePoses[kReadbackTextureCount];
