Fix build error in GoogleVRHMD

diff --git a/Engine/Plugins/Runtime/GoogleVR/GoogleVRHMD/Source/GoogleVRHMD/Private/GoogleVRHMD.cpp b/Engine/Plugins/Runtime/GoogleVR/GoogleVRHMD/Source/GoogleVRHMD/Private/GoogleVRHMD.cpp
--- a/Engine/Plugins/Runtime/GoogleVR/GoogleVRHMD/Source/GoogleVRHMD/Private/GoogleVRHMD.cpp
+++ b/Engine/Plugins/Runtime/GoogleVR/GoogleVRHMD/Source/GoogleVRHMD/Private/GoogleVRHMD.cpp
@@ -412,7 +412,7 @@ FGoogleVRHMD::FGoogleVRHMD(const FAutoRegister& AutoRegister)
 		bUseGVRApiDistortionCorrection = bUseOffscreenFramebuffers;
 		//bUseGVRApiDistortionCorrection = true;  //Uncomment this line is you want to use GVR distortion when async reprojection is not enabled.
 
-		// Query for direct multi-view
+		// Query for mobile multi-view
 		GSupportsMobileMultiView = gvr_is_feature_supported(GVRAPI, GVR_FEATURE_MULTIVIEW) && bUseOffscreenFramebuffers;
 		const auto CVarMobileMultiView = IConsoleManager::Get().FindTConsoleVariableDataInt(TEXT("vr.MobileMultiView"));
 		const bool bIsMobileMultiViewEnabled = (CVarMobileMultiView && CVarMobileMultiView->GetValueOnAnyThread() != 0);
@@ -561,7 +561,7 @@ void FGoogleVRHMD::UpdateGVRViewportList() const
 
 	ActiveViewportList = bDistortionCorrectionEnabled ? DistortedBufferViewportList : NonDistortedBufferViewportList;
 
-	if (IsMobileMultiViewDirect())
+	if (IsMobileMultiView())
 	{
 		check(gvr_buffer_viewport_list_get_size(ActiveViewportList) == 2);
 		for (uint32 EyeIndex = 0; EyeIndex < 2; ++EyeIndex)
@@ -756,7 +756,7 @@ bool FGoogleVRHMD::SetGVRHMDRenderTargetSize(int DesiredWidth, int DesiredHeight
 		return false;
 	}
 
-	const uint32 AdjustedDesiredWidth = (IsMobileMultiViewDirect()) ? DesiredWidth / 2 : DesiredWidth;
+	const uint32 AdjustedDesiredWidth = (IsMobileMultiView()) ? DesiredWidth / 2 : DesiredWidth;
 
 	// Ensure sizes are dividable by DividableBy to get post processing effects with lower resolution working well
 	const uint32 DividableBy = 4;
