Fix build error in GoogleVRHMD

diff --git a/Engine/Plugins/Runtime/GoogleVR/GoogleVRHMD/Source/GoogleVRHMD/Private/GoogleVRHMDRender.cpp b/Engine/Plugins/Runtime/GoogleVR/GoogleVRHMD/Source/GoogleVRHMD/Private/GoogleVRHMDRender.cpp
--- a/Engine/Plugins/Runtime/GoogleVR/GoogleVRHMD/Source/GoogleVRHMD/Private/GoogleVRHMDRender.cpp
+++ b/Engine/Plugins/Runtime/GoogleVR/GoogleVRHMD/Source/GoogleVRHMD/Private/GoogleVRHMDRender.cpp
@@ -316,7 +316,7 @@ bool FGoogleVRHMD::AllocateRenderTargetTexture(uint32 Index, uint32 SizeX, uint3
 #if GOOGLEVRHMD_SUPPORTED_PLATFORMS
 	if(CustomPresent)
 	{
-		const uint32 NumLayers = (IsMobileMultiViewDirect()) ? 2 : 1;
+		const uint32 NumLayers = (IsMobileMultiView()) ? 2 : 1;
 		bool Success = CustomPresent->AllocateRenderTargetTexture(Index, SizeX, SizeY, Format, NumLayers, NumMips, InFlags, TargetableTextureFlags);
 		if (Success)
 		{
@@ -493,7 +493,7 @@ void FGoogleVRHMDCustomPresent::CreateGVRSwapChain()
 	gvr_buffer_spec_set_depth_stencil_format(BufferSpec, GVR_DEPTH_STENCIL_FORMAT_NONE);
 	// We are using the default color buffer format in GVRSDK, which is RGBA8, and that is also the format passed in.
 
-	if (HMD->IsMobileMultiViewDirect())
+	if (HMD->IsMobileMultiView())
 	{
 		gvr_sizei BufferSize = gvr_buffer_spec_get_size(BufferSpec);
 		BufferSize.width /= 2;
@@ -638,4 +638,4 @@ bool FGoogleVRHMDCustomPresent::Present(int32& InOutSyncInterval)
 void FGoogleVRHMDCustomPresent::OnBackBufferResize()
 {
 }
-#endif // GOOGLEVRHMD_SUPPORTED_PLATFORMS
\ No newline at end of file
+#endif // GOOGLEVRHMD_SUPPORTED_PLATFORMS
diff --git a/Engine/Plugins/Runtime/GoogleVR/GoogleVRHMD/Source/GoogleVRHMD/Private/GoogleVRSplash.cpp b/Engine/Plugins/Runtime/GoogleVR/GoogleVRHMD/Source/GoogleVRHMD/Private/GoogleVRSplash.cpp
--- a/Engine/Plugins/Runtime/GoogleVR/GoogleVRHMD/Source/GoogleVRHMD/Private/GoogleVRSplash.cpp
+++ b/Engine/Plugins/Runtime/GoogleVR/GoogleVRHMD/Source/GoogleVRHMD/Private/GoogleVRSplash.cpp
@@ -70,7 +70,7 @@ void FGoogleVRSplash::AllocateSplashScreenRenderTarget()
 {
 	if (!GVRCustomPresent->TextureSet.IsValid())
 	{
-		const uint32 NumLayers = (GVRHMD->IsMobileMultiViewDirect()) ? 2 : 1;
+		const uint32 NumLayers = (GVRHMD->IsMobileMultiView()) ? 2 : 1;
 		GVRCustomPresent->AllocateRenderTargetTexture(0, GVRHMD->GVRRenderTargetSize.X, GVRHMD->GVRRenderTargetSize.Y, PF_B8G8R8A8, NumLayers, 1, TexCreate_None, TexCreate_RenderTargetable);
 	}
 }
