More build fixes

diff --git a/Engine/Plugins/Media/AndroidCamera/Source/AndroidCamera/Private/Player/AndroidCameraPlayer.cpp b/Engine/Plugins/Media/AndroidCamera/Source/AndroidCamera/Private/Player/AndroidCameraPlayer.cpp
--- a/Engine/Plugins/Media/AndroidCamera/Source/AndroidCamera/Private/Player/AndroidCameraPlayer.cpp
+++ b/Engine/Plugins/Media/AndroidCamera/Source/AndroidCamera/Private/Player/AndroidCameraPlayer.cpp
@@ -301,7 +301,7 @@ static void DoUpdateExternalCameraSampleExecute(TWeakPtr<FJavaAndroidCameraPlaye
 	FTextureRHIRef VideoTexture = PinnedJavaCameraPlayer->GetVideoTexture();
 	if (VideoTexture == nullptr)
 	{
-		FRHIResourceCreateInfo CreateInfo;
+		FRHIResourceCreateInfo CreateInfo(TEXT("VideoTexture"));
 		VideoTexture = GDynamicRHI->RHICreateTextureExternal2D(1, 1, PF_R8G8B8A8, 1, 1, TexCreate_None, ERHIAccess::Unknown, CreateInfo);
 		PinnedJavaCameraPlayer->SetVideoTexture(VideoTexture);
 
diff --git a/Engine/Plugins/Media/AndroidCamera/Source/AndroidCamera/Private/Player/AndroidCameraTextureSample.h b/Engine/Plugins/Media/AndroidCamera/Source/AndroidCamera/Private/Player/AndroidCameraTextureSample.h
--- a/Engine/Plugins/Media/AndroidCamera/Source/AndroidCamera/Private/Player/AndroidCameraTextureSample.h
+++ b/Engine/Plugins/Media/AndroidCamera/Source/AndroidCamera/Private/Player/AndroidCameraTextureSample.h
@@ -145,7 +145,7 @@ public:
 		const ETextureCreateFlags CreateFlags = TexCreate_Dynamic | TexCreate_SRGB;
 
 		TRefCountPtr<FRHITexture2D> DummyTexture2DRHI;
-		FRHIResourceCreateInfo CreateInfo;
+		FRHIResourceCreateInfo CreateInfo(TEXT("DummyTexture2D"));
 
 		RHICreateTargetableShaderResource2D(
 			Dim.X,
diff --git a/Engine/Plugins/Media/AndroidMedia/Source/AndroidMedia/Private/Player/AndroidMediaPlayer.cpp b/Engine/Plugins/Media/AndroidMedia/Source/AndroidMedia/Private/Player/AndroidMediaPlayer.cpp
--- a/Engine/Plugins/Media/AndroidMedia/Source/AndroidMedia/Private/Player/AndroidMediaPlayer.cpp
+++ b/Engine/Plugins/Media/AndroidMedia/Source/AndroidMedia/Private/Player/AndroidMediaPlayer.cpp
@@ -461,7 +461,7 @@ static void DoUpdateExternalMediaSampleExecute(TWeakPtr<FJavaAndroidMediaPlayer,
 	FTextureRHIRef VideoTexture = PinnedJavaMediaPlayer->GetVideoTexture();
 	if (VideoTexture == nullptr)
 	{
-		FRHIResourceCreateInfo CreateInfo;
+		FRHIResourceCreateInfo CreateInfo(TEXT("VideoTexture"));
 		VideoTexture = GDynamicRHI->RHICreateTextureExternal2D(1, 1, PF_R8G8B8A8, 1, 1, TexCreate_None, ERHIAccess::SRVGraphics, CreateInfo);
 		PinnedJavaMediaPlayer->SetVideoTexture(VideoTexture);
 
diff --git a/Engine/Plugins/Media/AndroidMedia/Source/AndroidMedia/Private/Player/AndroidMediaTextureSample.h b/Engine/Plugins/Media/AndroidMedia/Source/AndroidMedia/Private/Player/AndroidMediaTextureSample.h
--- a/Engine/Plugins/Media/AndroidMedia/Source/AndroidMedia/Private/Player/AndroidMediaTextureSample.h
+++ b/Engine/Plugins/Media/AndroidMedia/Source/AndroidMedia/Private/Player/AndroidMediaTextureSample.h
@@ -146,7 +146,7 @@ public:
 		const ETextureCreateFlags CreateFlags = TexCreate_Dynamic | TexCreate_SRGB;
 
 		TRefCountPtr<FRHITexture2D> DummyTexture2DRHI;
-		FRHIResourceCreateInfo CreateInfo;
+		FRHIResourceCreateInfo CreateInfo(TEXT("DummyTexture2D"));
 
 		RHICreateTargetableShaderResource2D(
 			Dim.X,
diff --git a/Engine/Plugins/Runtime/GoogleVR/GoogleVRHMD/Source/GoogleVRHMD/Private/GoogleVRHMDRender.cpp b/Engine/Plugins/Runtime/GoogleVR/GoogleVRHMD/Source/GoogleVRHMD/Private/GoogleVRHMDRender.cpp
--- a/Engine/Plugins/Runtime/GoogleVR/GoogleVRHMD/Source/GoogleVRHMD/Private/GoogleVRHMDRender.cpp
+++ b/Engine/Plugins/Runtime/GoogleVR/GoogleVRHMD/Source/GoogleVRHMD/Private/GoogleVRHMDRender.cpp
@@ -182,7 +182,7 @@ static void ResolvePendingRenderTarget(FRHICommandListImmediate& RHICmdList, FGr
 			FMemory::Memcpy(IndexBuffer.GetData(), Indices, InternalNumIndices * sizeof(uint16));
 
 			// Create index buffer. Fill buffer with initial data upon creation
-			FRHIResourceCreateInfo CreateInfo(&IndexBuffer);
+			FRHIResourceCreateInfo CreateInfo(TEXT("FGoogleVRHMD"), &IndexBuffer);
 			IndexBufferRHI = RHICreateIndexBuffer(sizeof(uint16), IndexBuffer.GetResourceDataSize(), BUF_Static, CreateInfo);
 		}
 	};
diff --git a/Engine/Source/Runtime/WebBrowser/Private/Android/AndroidWebBrowserWidget.cpp b/Engine/Source/Runtime/WebBrowser/Private/Android/AndroidWebBrowserWidget.cpp
--- a/Engine/Source/Runtime/WebBrowser/Private/Android/AndroidWebBrowserWidget.cpp
+++ b/Engine/Source/Runtime/WebBrowser/Private/Android/AndroidWebBrowserWidget.cpp
@@ -309,7 +309,7 @@ void SAndroidWebBrowserWidget::Tick(const FGeometry& AllottedGeometry, const dou
 					FTextureRHIRef VideoTexture = PinnedJavaWebBrowser->GetVideoTexture();
 					if (VideoTexture == nullptr)
 					{
-						FRHIResourceCreateInfo CreateInfo;
+						FRHIResourceCreateInfo CreateInfo(TEXT("VideoTexture"));
 						FIntPoint LocalSize = Params.Size;
 
 						VideoTexture = RHICreateTextureExternal2D(LocalSize.X, LocalSize.Y, PF_R8G8B8A8, 1, 1, TexCreate_None, CreateInfo);
